PropertyBindingUtils: Moved some property binding related code from StateTree to a shared plugin

diff --git a/Engine/Plugins/Runtime/PropertyBindingUtils/PropertyBindingUtils.uplugin b/Engine/Plugins/Runtime/PropertyBindingUtils/PropertyBindingUtils.uplugin
new file mode 100644
--- /dev/null
+++ b/Engine/Plugins/Runtime/PropertyBindingUtils/PropertyBindingUtils.uplugin
@@ -0,0 +1,37 @@
+{
+	"FileVersion" : 3,
+	"Version" : 1,
+	"VersionName" : "0.1",
+	"FriendlyName" : "PropertyBindingUtils",
+	"Description" : "Utility code for implementing property bindings",
+	"Category" : "Gameplay",
+	"CreatedBy" : "Epic Games, Inc.",
+	"CreatedByURL" : "https://epicgames.com",
+	"DocsURL" : "",
+	"MarketplaceURL" : "",
+	"SupportURL" : "",
+	"EnabledByDefault" : false,
+	"CanContainContent" : true,
+	"IsBetaVersion" : true,
+	"Installed" : false,
+	"Modules":
+	[
+		{
+			"Name": "PropertyBindingUtils",
+			"Type": "Runtime",
+			"LoadingPhase": "PreDefault"
+		},
+		{
+			"Name": "PropertyBindingUtilsTestSuite",
+			"Type": "UncookedOnly",
+			"LoadingPhase": "PreDefault"
+		}
+	],
+	"Plugins":
+	[
+		{
+			"Name": "StructUtils",
+			"Enabled": true
+		}
+	]
+}
\ No newline at end of file
diff --git a/Engine/Plugins/Runtime/PropertyBindingUtils/Source/PropertyBindingUtils/Private/PropertyBindingPath.cpp b/Engine/Plugins/Runtime/PropertyBindingUtils/Source/PropertyBindingUtils/Private/PropertyBindingPath.cpp
new file mode 100644
--- /dev/null
+++ b/Engine/Plugins/Runtime/PropertyBindingUtils/Source/PropertyBindingUtils/Private/PropertyBindingPath.cpp
@@ -0,0 +1,502 @@
+﻿// Copyright Epic Games, Inc. All Rights Reserved.
+#include "PropertyBindingPath.h"
+#include "UObject/EnumProperty.h"
+#include "Misc/EnumerateRange.h"
+#include "PropertyBag.h"
+
+#if WITH_EDITOR
+#include "UObject/CoreRedirects.h"
+#include "UObject/Package.h"
+#include "Engine/BlueprintGeneratedClass.h"
+#include "Engine/UserDefinedStruct.h"
+#include "Kismet2/StructureEditorUtils.h"
+#include "UObject/Field.h"
+#endif
+
+#include UE_INLINE_GENERATED_CPP_BY_NAME(PropertyBindingPath)
+
+//----------------------------------------------------------------//
+//  FPropertyBindingPath
+//----------------------------------------------------------------//
+
+bool FPropertyBindingPath::FromString(const FString& InPath)
+{
+	Segments.Reset();
+	
+	if (InPath.IsEmpty())
+	{
+		return true;
+	}
+	
+	bool bResult = true;
+	TArray<FString> PathSegments;
+	InPath.ParseIntoArray(PathSegments, TEXT("."), /*InCullEmpty=*/false);
+	
+	for (const FString& Segment : PathSegments)
+	{
+		if (Segment.IsEmpty())
+		{
+			bResult = false;
+			break;
+		}
+
+		int32 FirstBracket = INDEX_NONE;
+		int32 LastBracket = INDEX_NONE;
+		if (Segment.FindChar(TEXT('['), FirstBracket)
+			&& Segment.FindLastChar(TEXT(']'), LastBracket))
+		{
+			const int32 NameStringLength = FirstBracket;
+			const int32 IndexStringLength = LastBracket - FirstBracket - 1;
+			if (NameStringLength < 1
+				|| IndexStringLength <= 0)
+			{
+				bResult = false;
+				break;
+			}
+
+			const FString NameString = Segment.Left(FirstBracket);
+			const FString IndexString = Segment.Mid(FirstBracket + 1, IndexStringLength);
+			int32 ArrayIndex = INDEX_NONE;
+			LexFromString(ArrayIndex, *IndexString);
+			if (ArrayIndex < 0)
+			{
+				bResult = false;
+				break;
+			}
+			
+			AddPathSegment(FName(NameString), ArrayIndex);
+		}
+		else
+		{
+			AddPathSegment(FName(Segment));
+		}
+	}
+
+	if (!bResult)
+	{
+		Segments.Reset();
+	}
+	
+	return bResult;
+}
+
+bool FPropertyBindingPath::UpdateSegments(const UStruct* BaseStruct, FString* OutError)
+{
+	return UpdateSegmentsFromValue(FPropertyBindingDataView(BaseStruct, nullptr), OutError);
+}
+
+bool FPropertyBindingPath::UpdateSegmentsFromValue(const FPropertyBindingDataView BaseValueView, FString* OutError)
+{
+	TArray<FPropertyBindingPathIndirection> Indirections;
+	if (!ResolveIndirectionsWithValue(BaseValueView, Indirections, OutError, /*bHandleRedirects*/true))
+	{
+		return false;
+	}
+
+	for (FPropertyBindingPathSegment& Segment : Segments)
+	{
+		Segment.SetInstanceStruct(nullptr);
+	}
+	
+	for (const FPropertyBindingPathIndirection& Indirection : Indirections)
+	{
+		if (Indirection.InstanceStruct != nullptr)
+		{
+			Segments[Indirection.PathSegmentIndex].SetInstanceStruct(Indirection.InstanceStruct);
+		}
+#if WITH_EDITORONLY_DATA		
+		if (!Indirection.GetRedirectedName().IsNone())
+		{
+			Segments[Indirection.PathSegmentIndex].SetName(Indirection.GetRedirectedName());
+		}
+		Segments[Indirection.PathSegmentIndex].SetPropertyGuid(Indirection.GetPropertyGuid());
+#endif			
+	}
+
+	return true;
+}
+
+FString FPropertyBindingPath::ToString(const int32 HighlightedSegment, const TCHAR* HighlightPrefix, const TCHAR* HighlightPostfix, const bool bOutputInstances) const
+{
+	FString Result;
+	for (TEnumerateRef<const FPropertyBindingPathSegment> Segment : EnumerateRange(Segments))
+	{
+		if (Segment.GetIndex() > 0)
+		{
+			Result += TEXT(".");
+		}
+		if (Segment.GetIndex() == HighlightedSegment && HighlightPrefix)
+		{
+			Result += HighlightPrefix;
+		}
+
+		if (bOutputInstances && Segment->GetInstanceStruct())
+		{
+			Result += FString::Printf(TEXT("(%s)"), *GetNameSafe(Segment->GetInstanceStruct()));
+		}
+
+		if (Segment->GetArrayIndex() >= 0)
+		{
+			Result += FString::Printf(TEXT("%s[%d]"), *Segment->GetName().ToString(), Segment->GetArrayIndex());
+		}
+		else
+		{
+			Result += Segment->GetName().ToString();
+		}
+
+		if (Segment.GetIndex() == HighlightedSegment && HighlightPostfix)
+		{
+			Result += HighlightPostfix;
+		}
+	}
+	return Result;
+}
+
+
+bool FPropertyBindingPath::ResolveIndirections(const UStruct* BaseStruct, TArray<FPropertyBindingPathIndirection>& OutIndirections, FString* OutError, bool bHandleRedirects) const
+{
+	return ResolveIndirectionsWithValue(FPropertyBindingDataView(BaseStruct, nullptr), OutIndirections, OutError, bHandleRedirects);
+}
+
+bool FPropertyBindingPath::ResolveIndirectionsWithValue(const FPropertyBindingDataView BaseValueView, TArray<FPropertyBindingPathIndirection>& OutIndirections, FString* OutError, bool bHandleRedirects) const
+{
+	OutIndirections.Reset();
+	if (OutError)
+	{
+		OutError->Reset();
+	}
+	
+	// Nothing to do for an empty path.
+	if (IsPathEmpty())
+	{
+		return true;
+	}
+
+	const void* CurrentAddress = BaseValueView.GetMemory();
+	const UStruct* CurrentStruct = BaseValueView.GetStruct();
+	
+	for (const TEnumerateRef<const FPropertyBindingPathSegment> Segment : EnumerateRange(Segments))
+	{
+		if (CurrentStruct == nullptr)
+		{
+			if (OutError)
+			{
+				*OutError = FString::Printf(TEXT("Malformed path '%s'."),
+					*ToString(Segment.GetIndex(), TEXT("<"), TEXT(">")));
+			}
+			OutIndirections.Reset();
+			return false;
+		}
+
+		const FProperty* Property = CurrentStruct->FindPropertyByName(Segment->GetName());
+		const bool bWithValue = CurrentAddress != nullptr;
+
+#if WITH_EDITORONLY_DATA
+		FName RedirectedName;
+		FGuid PropertyGuid = Segment->GetPropertyGuid();
+
+		// Try to fix the path in editor.
+		if (bHandleRedirects)
+		{
+			
+			// Check if there's a core redirect for it.
+			if (!Property)
+			{
+				// Try to match by property ID (Blueprint or User Defined Struct).
+				if (Segment->GetPropertyGuid().IsValid())
+				{
+					if (const UBlueprintGeneratedClass* BlueprintClass = Cast<UBlueprintGeneratedClass>(CurrentStruct))
+					{
+						if (const FName* Name = BlueprintClass->PropertyGuids.FindKey(Segment->GetPropertyGuid()))
+						{
+							RedirectedName = *Name;
+							Property = CurrentStruct->FindPropertyByName(RedirectedName);
+						}
+					}
+					else if (const UUserDefinedStruct* UserDefinedStruct = Cast<UUserDefinedStruct>(CurrentStruct))
+					{
+						if (FProperty* FoundProperty = FStructureEditorUtils::GetPropertyByGuid(UserDefinedStruct, Segment->GetPropertyGuid()))
+						{
+							RedirectedName = FoundProperty->GetFName();
+							Property = FoundProperty;
+						}
+					}
+					else if (const UPropertyBag* PropertyBag = Cast<UPropertyBag>(CurrentStruct))
+					{
+						if (const FPropertyBagPropertyDesc* Desc = PropertyBag->FindPropertyDescByID(Segment->GetPropertyGuid()))
+						{
+							if (Desc->CachedProperty)
+							{
+								RedirectedName = Desc->CachedProperty->GetFName();
+								Property = Desc->CachedProperty;
+							}
+						}
+					}
+				}
+				else
+				{
+					// Try core redirect
+					const FCoreRedirectObjectName OldPropertyName(Segment->GetName(), CurrentStruct->GetFName(), *CurrentStruct->GetOutermost()->GetPathName());
+					const FCoreRedirectObjectName NewPropertyName = FCoreRedirects::GetRedirectedName(ECoreRedirectFlags::Type_Property, OldPropertyName);
+					if (OldPropertyName != NewPropertyName)
+					{
+						// Cached the result for later use.
+						RedirectedName = NewPropertyName.ObjectName;
+
+						Property = CurrentStruct->FindPropertyByName(RedirectedName);
+					}
+				}
+			}
+
+			// Update PropertyGuid 
+			if (Property)
+			{
+				const FName PropertyName = !RedirectedName.IsNone() ? RedirectedName : Segment->GetName();
+				if (const UBlueprintGeneratedClass* BlueprintClass = Cast<UBlueprintGeneratedClass>(CurrentStruct))
+				{
+					if (const FGuid* VarGuid = BlueprintClass->PropertyGuids.Find(PropertyName))
+					{
+						PropertyGuid = *VarGuid;
+					}
+				}
+				else if (const UUserDefinedStruct* UserDefinedStruct = Cast<UUserDefinedStruct>(CurrentStruct))
+				{
+					// Parse Guid from UDS property name.
+					PropertyGuid = FStructureEditorUtils::GetGuidFromPropertyName(PropertyName);
+				}
+				else if (const UPropertyBag* PropertyBag = Cast<UPropertyBag>(CurrentStruct))
+				{
+					if (const FPropertyBagPropertyDesc* Desc = PropertyBag->FindPropertyDescByPropertyName(PropertyName))
+					{
+						PropertyGuid = Desc->ID;
+					}
+				}
+			}
+		}
+#endif // WITH_EDITORONLY_DATA
+
+		if (!Property)
+		{
+			if (OutError)
+			{
+				*OutError = FString::Printf(TEXT("Malformed path '%s', could not find property '%s%s::%s'."),
+					*ToString(Segment.GetIndex(), TEXT("<"), TEXT(">")),
+					CurrentStruct->GetPrefixCPP(), *CurrentStruct->GetName(), *Segment->GetName().ToString());
+			}
+			OutIndirections.Reset();
+			return false;
+		}
+
+		const FArrayProperty* ArrayProperty = CastField<FArrayProperty>(Property);
+		int ArrayIndex = 0;
+		int32 Offset = 0;
+		if (ArrayProperty && Segment->GetArrayIndex() != INDEX_NONE)
+		{
+			FPropertyBindingPathIndirection& Indirection = OutIndirections.AddDefaulted_GetRef();
+			Indirection.Property = Property;
+			Indirection.ContainerAddress = CurrentAddress;
+			Indirection.ContainerStruct = CurrentStruct;
+			Indirection.InstanceStruct = nullptr;
+			Indirection.ArrayIndex = Segment->GetArrayIndex();
+			Indirection.PropertyOffset = ArrayProperty->GetOffset_ForInternal();
+			Indirection.PathSegmentIndex = Segment.GetIndex();
+			Indirection.AccessType = EPropertyBindingAccessType::IndexArray;
+#if WITH_EDITORONLY_DATA
+			Indirection.RedirectedName = RedirectedName;
+			Indirection.PropertyGuid = PropertyGuid;
+#endif
+			
+			ArrayIndex = 0;
+			Offset = 0;
+			Property = ArrayProperty->Inner;
+
+			if (bWithValue)
+			{
+				FScriptArrayHelper Helper(ArrayProperty, (uint8*)CurrentAddress + ArrayProperty->GetOffset_ForInternal());
+				if (!Helper.IsValidIndex(Segment->GetArrayIndex()))
+				{
+					if (OutError)
+					{
+						*OutError = FString::Printf(TEXT("Index %d out of range (num elements %d) trying to access dynamic array '%s'."),
+							Segment->GetArrayIndex(), Helper.Num(), *ToString(Segment.GetIndex(), TEXT("<"), TEXT(">")));
+					}
+					OutIndirections.Reset();
+					return false;
+				}
+				CurrentAddress = Helper.GetRawPtr(Segment->GetArrayIndex());
+			}
+		}
+		else
+		{
+			if (Segment->GetArrayIndex() > Property->ArrayDim)
+			{
+				if (OutError)
+				{
+					*OutError = FString::Printf(TEXT("Index %d out of range %d trying to access static array '%s'."),
+						Segment->GetArrayIndex(), Property->ArrayDim, *ToString(Segment.GetIndex(), TEXT("<"), TEXT(">")));
+				}
+				OutIndirections.Reset();
+				return false;
+			}
+			ArrayIndex = FMath::Max(0, Segment->GetArrayIndex());
+			Offset = Property->GetOffset_ForInternal() + Property->ElementSize * ArrayIndex;
+		}
+
+		FPropertyBindingPathIndirection& Indirection = OutIndirections.AddDefaulted_GetRef();
+		Indirection.Property = Property;
+		Indirection.ContainerAddress = CurrentAddress;
+		Indirection.ContainerStruct = CurrentStruct;
+		Indirection.ArrayIndex = ArrayIndex;
+		Indirection.PropertyOffset = Offset;
+		Indirection.PathSegmentIndex = Segment.GetIndex();
+		Indirection.AccessType = EPropertyBindingAccessType::Offset; 
+#if WITH_EDITORONLY_DATA
+		Indirection.RedirectedName = RedirectedName;
+		Indirection.PropertyGuid = PropertyGuid;
+#endif
+		const bool bLastSegment = Segment.GetIndex() == (Segments.Num() - 1);
+
+		if (!bLastSegment)
+		{
+			if (const FStructProperty* StructProperty = CastField<FStructProperty>(Property))
+			{
+				if (bWithValue)
+				{
+					if (StructProperty->Struct == TBaseStructure<FInstancedStruct>::Get())
+					{
+						// The property path is pointing into the instanced struct, it must be present.
+						// @TODO:	We could potentially check the BaseStruct metadata in editor (for similar behavior as objects)
+						//			Omitting for now to have matching functionality in editor and runtime.
+						const FInstancedStruct& InstancedStruct = *reinterpret_cast<const FInstancedStruct*>((uint8*)CurrentAddress + Offset);
+						if (!InstancedStruct.IsValid())
+						{
+							if (OutError)
+							{
+								*OutError = FString::Printf(TEXT("Expecting valid instanced struct value at path '%s'."),
+									*ToString(Segment.GetIndex(), TEXT("<"), TEXT(">")));
+							}
+							OutIndirections.Reset();
+							return false;
+						}
+						const UScriptStruct* ValueInstanceStructType = InstancedStruct.GetScriptStruct();
+
+						CurrentAddress = InstancedStruct.GetMemory();
+						CurrentStruct = ValueInstanceStructType;
+						Indirection.InstanceStruct = CurrentStruct;
+						Indirection.AccessType = EPropertyBindingAccessType::StructInstance; 
+					}
+					else
+					{
+						CurrentAddress = (uint8*)CurrentAddress + Offset;
+						CurrentStruct = StructProperty->Struct;
+						Indirection.AccessType = EPropertyBindingAccessType::Offset;
+					}
+				}
+				else
+				{
+					if (Segment->GetInstanceStruct())
+					{
+						CurrentStruct = Segment->GetInstanceStruct();
+						Indirection.InstanceStruct = CurrentStruct;
+						Indirection.AccessType = EPropertyBindingAccessType::StructInstance;
+					}
+					else
+					{
+						CurrentStruct = StructProperty->Struct;
+						Indirection.AccessType = EPropertyBindingAccessType::Offset;
+					}
+				}
+			}
+			else if (const FObjectProperty* ObjectProperty = CastField<FObjectProperty>(Property))
+			{
+				if (bWithValue)
+				{
+					const UObject* Object = *reinterpret_cast<UObject* const*>((uint8*)CurrentAddress + Offset);
+					CurrentAddress = Object;
+					
+					// The property path is pointing into the object, if the object is present use it's specific type, otherwise use the type of the pointer.
+					if (Object)
+					{
+						CurrentStruct = Object->GetClass();
+						Indirection.InstanceStruct = CurrentStruct;
+						Indirection.AccessType = EPropertyBindingAccessType::ObjectInstance;
+					}
+					else
+					{
+						CurrentStruct = ObjectProperty->PropertyClass;
+						Indirection.AccessType = EPropertyBindingAccessType::Object;
+					}
+				}
+				else
+				{
+					if (Segment->GetInstanceStruct())
+					{
+						CurrentStruct = Segment->GetInstanceStruct();
+						Indirection.InstanceStruct = CurrentStruct;
+						Indirection.AccessType = EPropertyBindingAccessType::ObjectInstance;
+					}
+					else
+					{
+						CurrentStruct = ObjectProperty->PropertyClass;
+						Indirection.AccessType = EPropertyBindingAccessType::Object;
+					}
+				}
+			}
+			// Check to see if this is a simple weak object property (eg. not an array of weak objects).
+			else if (const FWeakObjectProperty* WeakObjectProperty = CastField<FWeakObjectProperty>(Property))
+			{
+				const TWeakObjectPtr<UObject>& WeakObjectPtr = *reinterpret_cast<const TWeakObjectPtr<UObject>*>((uint8*)CurrentAddress + Offset);
+				const UObject* Object = WeakObjectPtr.Get();
+				CurrentAddress = Object;
+				CurrentStruct = WeakObjectProperty->PropertyClass;
+				Indirection.AccessType = EPropertyBindingAccessType::WeakObject;
+			}
+			// Check to see if this is a simple soft object property (eg. not an array of soft objects).
+			else if (const FSoftObjectProperty* SoftObjectProperty = CastField<FSoftObjectProperty>(Property))
+			{
+				const FSoftObjectPtr& SoftObjectPtr = *reinterpret_cast<const FSoftObjectPtr*>((uint8*)CurrentAddress + Offset);
+				const UObject* Object = SoftObjectPtr.Get();
+				CurrentAddress = Object;
+				CurrentStruct = SoftObjectProperty->PropertyClass;
+				Indirection.AccessType = EPropertyBindingAccessType::SoftObject;
+			}
+			else
+			{
+				// We get here if we encounter a property type that is not supported for indirection (e.g. Map or Set).
+				if (OutError)
+				{
+					*OutError = FString::Printf(TEXT("Unsupported property indirection type %s in path '%s'."),
+						*Property->GetCPPType(), *ToString(Segment.GetIndex(), TEXT("<"), TEXT(">")));
+				}
+				OutIndirections.Reset();
+				return false;
+			}
+		}
+	}
+
+	return true;
+}
+
+bool FPropertyBindingPath::operator==(const FPropertyBindingPath& RHS) const
+{
+#if WITH_EDITORONLY_DATA
+	if (StructID != RHS.StructID)
+	{
+		return false;
+	}
+#endif // WITH_EDITORONLY_DATA
+	if (Segments.Num() != RHS.Segments.Num())
+	{
+		return false;
+	}
+
+	for (TEnumerateRef<const FPropertyBindingPathSegment> Segment : EnumerateRange(Segments))
+	{
+		if (*Segment != RHS.Segments[Segment.GetIndex()])
+		{
+			return false;
+		}
+	}
+
+	return true;
+}
\ No newline at end of file
diff --git a/Engine/Plugins/Runtime/PropertyBindingUtils/Source/PropertyBindingUtils/Private/PropertyBindingUtilsModule.cpp b/Engine/Plugins/Runtime/PropertyBindingUtils/Source/PropertyBindingUtils/Private/PropertyBindingUtilsModule.cpp
new file mode 100644
--- /dev/null
+++ b/Engine/Plugins/Runtime/PropertyBindingUtils/Source/PropertyBindingUtils/Private/PropertyBindingUtilsModule.cpp
@@ -0,0 +1,23 @@
+// Copyright Epic Games, Inc. All Rights Reserved.
+
+#include "PropertyBindingUtilsModule.h"
+
+#define LOCTEXT_NAMESPACE "PropertyBindingUtils"
+
+class FPropertyBindingUtilsModule : public IPropertyBindingUtilsModule
+{
+	virtual void StartupModule() override;
+	virtual void ShutdownModule() override;
+};
+
+IMPLEMENT_MODULE(FPropertyBindingUtilsModule, PropertyBindingUtils)
+
+void FPropertyBindingUtilsModule::StartupModule()
+{
+}
+
+void FPropertyBindingUtilsModule::ShutdownModule()
+{
+}
+
+#undef LOCTEXT_NAMESPACE
diff --git a/Engine/Plugins/Runtime/PropertyBindingUtils/Source/PropertyBindingUtils/PropertyBindingUtils.Build.cs b/Engine/Plugins/Runtime/PropertyBindingUtils/Source/PropertyBindingUtils/PropertyBindingUtils.Build.cs
new file mode 100644
--- /dev/null
+++ b/Engine/Plugins/Runtime/PropertyBindingUtils/Source/PropertyBindingUtils/PropertyBindingUtils.Build.cs
@@ -0,0 +1,36 @@
+// Copyright Epic Games, Inc. All Rights Reserved.
+
+namespace UnrealBuildTool.Rules
+{
+	public class PropertyBindingUtils : ModuleRules
+	{
+		public PropertyBindingUtils(ReadOnlyTargetRules Target) : base(Target)
+		{
+			PublicDependencyModuleNames.AddRange(
+			new string[] {
+				"Core",
+				"CoreUObject",
+				"Engine",
+				"StructUtils",
+			}
+			);
+
+			PrivateDependencyModuleNames.AddRange(
+			new string[] {
+				"RenderCore",
+				"InputCore"
+			}
+			);
+
+			if (Target.bBuildEditor)
+			{
+				PublicDependencyModuleNames.AddRange(
+					new string[] {
+						"UnrealEd",
+						"BlueprintGraph",
+					}
+				);
+			}
+		}
+	}
+}
diff --git a/Engine/Plugins/Runtime/PropertyBindingUtils/Source/PropertyBindingUtils/Public/PropertyBindingPath.h b/Engine/Plugins/Runtime/PropertyBindingUtils/Source/PropertyBindingUtils/Public/PropertyBindingPath.h
new file mode 100644
--- /dev/null
+++ b/Engine/Plugins/Runtime/PropertyBindingUtils/Source/PropertyBindingUtils/Public/PropertyBindingPath.h
@@ -0,0 +1,568 @@
+﻿// Copyright Epic Games, Inc. All Rights Reserved.
+
+#pragma once
+
+#include "EngineDefines.h"
+#include "Engine/EngineTypes.h"
+#include "StructView.h"
+#include "Misc/Guid.h"
+#include "PropertyBindingPath.generated.h"
+
+struct FPropertyBindingPath;
+
+UENUM()
+enum class EPropertyBindingAccessType : uint8
+{
+	Offset,			// Access node is a simple basePtr + offset
+	Object,			// Access node needs to dereference an object at its current address
+	WeakObject,		// Access is a weak object
+	SoftObject,		// Access is a soft object
+	ObjectInstance,	// Access node needs to dereference an object of specific type at its current address
+	StructInstance,	// Access node needs to dereference an instanced struct of specific type at its current address
+	IndexArray,		// Access node indexes a dynamic array
+};
+
+/**
+ * Short lived pointer to an object or struct.
+ * The data view expects a type (UStruct) when you pass in a valid memory. In case of null, the type can be empty too.
+ */
+struct PROPERTYBINDINGUTILS_API FPropertyBindingDataView
+{
+	FPropertyBindingDataView() = default;
+
+	// Generic struct constructor.
+	FPropertyBindingDataView(const UStruct* InStruct, void* InMemory) : Struct(InStruct), Memory(InMemory)
+	{
+		// Must have type with valid pointer.
+		check(!Memory || (Memory && Struct));
+	}
+
+	// UObject constructor.
+	FPropertyBindingDataView(UObject* Object) : Struct(Object ? Object->GetClass() : nullptr), Memory(Object)
+	{
+		// Must have type with valid pointer.
+		check(!Memory || (Memory && Struct));
+	}
+
+	// Struct from a StructView.
+	FPropertyBindingDataView(FStructView StructView) : Struct(StructView.GetScriptStruct()), Memory(StructView.GetMemory())
+	{
+		// Must have type with valid pointer.
+		check(!Memory || (Memory && Struct));
+	}
+
+	/**
+	 * Check is the view is valid (both pointer and type are set). On valid views it is safe to call the Get<>() methods returning a reference.
+	 * @return True if the view is valid.
+	*/
+	bool IsValid() const
+	{
+		return Memory != nullptr && Struct != nullptr;
+	}
+
+	/*
+	 * UObject getters (reference & pointer, const & mutable)
+	 */
+	template <typename T>
+    typename TEnableIf<TIsDerivedFrom<T, UObject>::IsDerived, const T&>::Type Get() const
+	{
+		check(Memory != nullptr);
+		check(Struct != nullptr);
+		check(Struct->IsChildOf(T::StaticClass()));
+		return *((T*)Memory);
+	}
+
+	template <typename T>
+	typename TEnableIf<TIsDerivedFrom<T, UObject>::IsDerived, T&>::Type GetMutable() const
+	{
+		check(Memory != nullptr);
+		check(Struct != nullptr);
+		check(Struct->IsChildOf(T::StaticClass()));
+		return *((T*)Memory);
+	}
+
+	template <typename T>
+	typename TEnableIf<TIsDerivedFrom<T, UObject>::IsDerived, const T*>::Type GetPtr() const
+	{
+		// If Memory is set, expect Struct too. Otherwise, let nulls pass through.
+		check(!Memory || (Memory && Struct));
+		check(!Struct || Struct->IsChildOf(T::StaticClass()));
+		return ((T*)Memory);
+	}
+
+	template <typename T>
+    typename TEnableIf<TIsDerivedFrom<T, UObject>::IsDerived, T*>::Type GetMutablePtr() const
+	{
+		// If Memory is set, expect Struct too. Otherwise, let nulls pass through.
+		check(!Memory || (Memory && Struct));
+		check(!Struct || Struct->IsChildOf(T::StaticClass()));
+		return ((T*)Memory);
+	}
+
+	/*
+	 * Struct getters (reference & pointer, const & mutable)
+	 */
+	template <typename T>
+	typename TEnableIf<!TIsDerivedFrom<T, UObject>::IsDerived && !TIsIInterface<T>::Value, const T&>::Type Get() const
+	{
+		check(Memory != nullptr);
+		check(Struct != nullptr);
+		check(Struct->IsChildOf(T::StaticStruct()));
+		return *((T*)Memory);
+	}
+
+	template <typename T>
+    typename TEnableIf<!TIsDerivedFrom<T, UObject>::IsDerived && !TIsIInterface<T>::Value, T&>::Type GetMutable() const
+	{
+		check(Memory != nullptr);
+		check(Struct != nullptr);
+		check(Struct->IsChildOf(T::StaticStruct()));
+		return *((T*)Memory);
+	}
+
+	template <typename T>
+    typename TEnableIf<!TIsDerivedFrom<T, UObject>::IsDerived && !TIsIInterface<T>::Value, const T*>::Type GetPtr() const
+	{
+		// If Memory is set, expect Struct too. Otherwise, let nulls pass through.
+		check(!Memory || (Memory && Struct));
+		check(!Struct || Struct->IsChildOf(T::StaticStruct()));
+		return ((T*)Memory);
+	}
+
+	template <typename T>
+    typename TEnableIf<!TIsDerivedFrom<T, UObject>::IsDerived && !TIsIInterface<T>::Value, T*>::Type GetMutablePtr() const
+	{
+		// If Memory is set, expect Struct too. Otherwise, let nulls pass through.
+		check(!Memory || (Memory && Struct));
+		check(!Struct || Struct->IsChildOf(T::StaticStruct()));
+		return ((T*)Memory);
+	}
+
+	/*
+	 * IInterface getters (reference & pointer, const & mutable)
+	 */
+	template <typename T>
+	typename TEnableIf<TIsIInterface<T>::Value, const T&>::Type Get() const
+	{
+		check(!Memory || (Memory && Struct));
+		check(Struct->IsChildOf(UObject::StaticClass()) && ((UClass*)Struct)->ImplementsInterface(T::UClassType::StaticClass()));
+		return *(T*)((UObject*)Memory)->GetInterfaceAddress(T::UClassType::StaticClass());
+	}
+
+	template <typename T>
+    typename TEnableIf<TIsIInterface<T>::Value, T&>::Type GetMutable() const
+	{
+		check(!Memory || (Memory && Struct));
+		check(Struct->IsChildOf(UObject::StaticClass()) && ((UClass*)Struct)->ImplementsInterface(T::UClassType::StaticClass()));
+		return *(T*)((UObject*)Memory)->GetInterfaceAddress(T::UClassType::StaticClass());
+	}
+
+	template <typename T>
+    typename TEnableIf<TIsIInterface<T>::Value, const T*>::Type GetPtr() const
+	{
+		check(!Memory || (Memory && Struct));
+		check(Struct->IsChildOf(UObject::StaticClass()) && ((UClass*)Struct)->ImplementsInterface(T::UClassType::StaticClass()));
+		return (T*)((UObject*)Memory)->GetInterfaceAddress(T::UClassType::StaticClass());
+	}
+
+	template <typename T>
+    typename TEnableIf<TIsIInterface<T>::Value, T*>::Type GetMutablePtr() const
+	{
+		check(!Memory || (Memory && Struct));
+		check(Struct->IsChildOf(UObject::StaticClass()) && ((UClass*)Struct)->ImplementsInterface(T::UClassType::StaticClass()));
+		return (T*)((UObject*)Memory)->GetInterfaceAddress(T::UClassType::StaticClass());
+	}
+
+	/** @return Struct describing the data type. */
+	const UStruct* GetStruct() const
+	{
+		return Struct;
+	}
+
+	/** @return Raw const pointer to the data. */
+	const void* GetMemory() const
+	{
+		return Memory;
+	}
+
+	/** @return Raw mutable pointer to the data. */
+	void* GetMutableMemory() const
+	{
+		return Memory;
+	}
+	
+protected:
+	/** UClass or UScriptStruct of the data. */
+	const UStruct* Struct = nullptr;
+
+	/** Memory pointing at the class or struct */
+	void* Memory = nullptr;
+};
+
+
+
+/**
+ * Struct describing an indirection at specific segment at path.
+ * Returned by FPropertyBindingPath::ResolveIndirections() and FPropertyBindingPath::ResolveIndirectionsWithValue().
+ * Generally there's one indirection per FProperty. Containers have one path segment but two indirection (container + inner type).
+ */
+struct PROPERTYBINDINGUTILS_API FPropertyBindingPathIndirection
+{
+	FPropertyBindingPathIndirection() = default;
+	explicit FPropertyBindingPathIndirection(const UStruct* InContainerStruct)
+		: ContainerStruct(InContainerStruct)
+	{
+	}
+
+	const FProperty* GetProperty() const
+	{
+		return Property;
+	}
+
+	const void* GetContainerAddress() const
+	{
+		return ContainerAddress;
+	}
+
+	const UStruct* GetInstanceStruct() const
+	{
+		return InstanceStruct;
+	}
+
+	const UStruct* GetContainerStruct() const
+	{
+		return ContainerStruct;
+	}
+
+	int32 GetArrayIndex() const
+	{
+		return ArrayIndex;
+	}
+
+	int32 GetPropertyOffset() const
+	{
+		return PropertyOffset;
+	}
+
+	int32 GetPathSegmentIndex() const
+	{
+		return PathSegmentIndex;
+	}
+
+	EPropertyBindingAccessType GetAccessType() const
+	{
+		return AccessType;
+	}
+	
+	const void* GetPropertyAddress() const
+	{
+		return (uint8*)ContainerAddress + PropertyOffset;
+	}
+
+#if WITH_EDITORONLY_DATA
+	FName GetRedirectedName() const
+	{
+		return RedirectedName;
+	}
+	
+	FGuid GetPropertyGuid() const
+	{
+		return PropertyGuid;
+	}
+#endif
+	
+private:
+	/** Property at the indirection. */
+	const FProperty* Property = nullptr;
+	
+	/** Address of the container class/struct where the property belongs to. Only valid if created with ResolveIndirectionsWithValue() */
+	const void* ContainerAddress = nullptr;
+	
+	/** Type of the container class/struct. */
+	const UStruct* ContainerStruct = nullptr;
+	
+	/** Type of the instance class/struct of when AccessType is ObjectInstance or StructInstance. */
+	const UStruct* InstanceStruct = nullptr;
+
+	/** Array index for static and dynamic arrays. Note: static array indexing is baked in the PropertyOffset. */
+	int32 ArrayIndex = 0;
+	
+	/** Offset of the property relative to ContainerAddress. Includes static array indexing. */
+	int32 PropertyOffset = 0;
+	
+	/** Index of the path segment where indirection originated from. */
+	int32 PathSegmentIndex = 0;
+	
+	/** How to access the data through the indirection. */
+	EPropertyBindingAccessType AccessType = EPropertyBindingAccessType::Offset;
+
+#if WITH_EDITORONLY_DATA
+	/** Redirected name, if the give property name was not found but was reconciled using core redirect or property Guid. Requires ResolveIndirections/WithValue() to be called with bHandleRedirects = true. */
+	FName RedirectedName;
+
+	/** Guid of the property for Blueprint classes or User Defined Structs. Requires ResolveIndirections/WithValue() to be called with bHandleRedirects = true. */
+	FGuid PropertyGuid;
+#endif
+
+	friend FPropertyBindingPath;
+};
+
+/** Struct describing a path segment in FPropertyBindingPath. */
+USTRUCT()
+struct PROPERTYBINDINGUTILS_API FPropertyBindingPathSegment
+{
+	GENERATED_BODY()
+
+	FPropertyBindingPathSegment() = default;
+
+	explicit FPropertyBindingPathSegment(const FName InName, const int32 InArrayIndex = INDEX_NONE, const UStruct* InInstanceStruct = nullptr)
+		: Name(InName)
+		, ArrayIndex(InArrayIndex)
+		, InstanceStruct(InInstanceStruct)
+	{
+	}
+
+	bool operator==(const FPropertyBindingPathSegment& RHS) const
+	{
+		return Name == RHS.Name && InstanceStruct == RHS.InstanceStruct && ArrayIndex == RHS.ArrayIndex;
+	}
+
+	bool operator!=(const FPropertyBindingPathSegment& RHS) const
+	{
+		return !(*this == RHS);
+	}
+
+	void SetName(const FName InName)
+	{
+		Name = InName;
+	}
+
+	FName GetName() const
+	{
+		return Name;
+	}
+
+	void SetArrayIndex(const int32 InArrayIndex)
+	{
+		ArrayIndex = InArrayIndex;
+	}
+
+	int32 GetArrayIndex() const
+	{
+		return ArrayIndex;
+	}
+
+	void SetInstanceStruct(const UStruct* InInstanceStruct)
+	{
+		InstanceStruct = InInstanceStruct;
+	}
+
+	const UStruct* GetInstanceStruct() const
+	{
+		return InstanceStruct;
+	}
+
+#if WITH_EDITORONLY_DATA
+	FGuid GetPropertyGuid() const
+	{
+		return PropertyGuid;
+	}
+
+	void SetPropertyGuid(const FGuid NewGuid)
+	{
+		PropertyGuid = NewGuid;
+	}
+#endif
+	
+private:
+	/** Name of the property */
+	UPROPERTY()
+	FName Name;
+
+	/** Array index if the property is dynamic or static array. */
+	UPROPERTY()
+	int32 ArrayIndex = INDEX_NONE;
+
+	/** Type of the instanced struct or object reference by the property at the segment. This allows the path to be resolved when it points to a specific instance. */
+	UPROPERTY()
+	TObjectPtr<const UStruct> InstanceStruct = nullptr;
+
+#if WITH_EDITORONLY_DATA
+	/** Guid of the property for Blueprint classes, User Defined Structs, or Property Bags. */
+	FGuid PropertyGuid;
+#endif
+};
+
+
+/**
+ * Representation of a property path that can be used for property access and binding.
+ *
+ * The engine supports many types of property paths, this implementation has these specific properties:
+ *		- Allow to resolve all the indirections from a base value (object or struct) up to the leaf property
+ *		- handle redirects from Core Redirect, BP classes, User Defines Structs and Property Bags
+ *
+ * You may also take a look at: FCachedPropertyPath, TFieldPath<>, FPropertyPath.
+ */
+USTRUCT()
+struct PROPERTYBINDINGUTILS_API FPropertyBindingPath
+{
+	GENERATED_BODY()
+
+	FPropertyBindingPath() = default;
+
+#if WITH_EDITORONLY_DATA
+	explicit FPropertyBindingPath(const FGuid InStructID)
+		: StructID(InStructID)
+	{
+	}
+
+	explicit FPropertyBindingPath(const FGuid InStructID, const FName PropertyName)
+		: StructID(InStructID)
+	{
+		Segments.Emplace(PropertyName);
+	}
+	
+	explicit FPropertyBindingPath(const FGuid InStructID, TConstArrayView<FPropertyBindingPathSegment> InSegments)
+		: StructID(InStructID)
+		, Segments(InSegments)
+	{
+	}
+#endif // WITH_EDITORONLY_DATA
+	
+	/**
+	 * Parses path from string. The path should be in format: Foo.Bar[1].Baz
+	 * @param InPath Path string to parse
+	 * @return true if path was parsed successfully.
+	 */
+	bool FromString(const FString& InPath);
+
+	/**
+	 * Returns the property path as a one string. Highlight allows to decorate a specific segment.
+	 * @param HighlightedSegment Index of the highlighted path segment
+	 * @param HighlightPrefix String to append before highlighted segment
+	 * @param HighlightPostfix String to append after highlighted segment
+	 * @param bOutputInstances if true, the instance struct types will be output. 
+	 */
+	FString ToString(const int32 HighlightedSegment = INDEX_NONE, const TCHAR* HighlightPrefix = nullptr, const TCHAR* HighlightPostfix = nullptr, const bool bOutputInstances = false) const;
+
+	/**
+	 * Resolves the property path against base struct type. The path is assumed to be relative to the BaseStruct.
+	 * @param BaseStruct Base struct/class type the path is relative to.
+	 * @param OutIndirections Indirections describing how the properties were accessed. 
+	 * @param OutError Optional, pointer to string where error will be logged if update fails.
+	 * @param bHandleRedirects If true, the method will try to resolve missing properties using core redirects, and properties on Blueprint and User Defined Structs by ID. Available only in editor builds!
+	 * @return true of path could be resolved against the base value, and instance types were updated.
+	 */
+	bool ResolveIndirections(const UStruct* BaseStruct, TArray<FPropertyBindingPathIndirection>& OutIndirections, FString* OutError = nullptr, bool bHandleRedirects = false) const;
+	
+	/**
+	 * Resolves the property path against base value. The path is assumed to be relative to the BaseValueView.
+	 * @param BaseValueView Base value the path is relative to.
+	 * @param OutIndirections Indirections describing how the properties were accessed. 
+	 * @param OutError Optional, pointer to string where error will be logged if update fails.
+	 * @param bHandleRedirects If true, the method will try to resolve missing properties using core redirects, and properties on Blueprint and User Defined Structs by ID. Available only in editor builds!
+	 * @return true of path could be resolved against the base value, and instance types were updated.
+	 */
+	bool ResolveIndirectionsWithValue(const FPropertyBindingDataView BaseValueView, TArray<FPropertyBindingPathIndirection>& OutIndirections, FString* OutError = nullptr, bool bHandleRedirects = false) const;
+
+	/**
+	 * Updates property segments from base struct type. The path is expected to be relative to the BaseStruct.
+	 * The method handles renamed properties (core redirect, Blueprint, User Defined Structs and Property Bags by ID).
+	 * @param BaseStruct Base value the path is relative to.
+	 * @param OutError Optional, pointer to string where error will be logged if update fails.
+	 * @return true of path could be resolved against the base value, and instance types were updated.
+	 */
+	bool UpdateSegments(const UStruct* BaseStruct, FString* OutError = nullptr);
+
+	/**
+	 * Updates property segments from base value. The path is expected to be relative to the base value.
+	 * The method updates instance types, and handles renamed properties (core redirect, Blueprint, User Defined Structs and Property Bags by ID).
+	 * By storing the instance types on the path, we can resolve the path without the base value later.
+	 * @param BaseValueView Base value the path is relative to.
+	 * @param OutError Optional, pointer to string where error will be logged if update fails.
+	 * @return true of path could be resolved against the base value, and instance types were updated.
+	 */
+	bool UpdateSegmentsFromValue(const FPropertyBindingDataView BaseValueView, FString* OutError = nullptr);
+
+	/** @return true if the path is empty. In that case the path points to the struct. */
+	bool IsPathEmpty() const
+	{
+		return Segments.IsEmpty();
+	}
+
+	/** @return true if any of the path segments is and indirection via instanced struct or object. */
+	bool HasAnyInstancedIndirection() const
+	{
+		return Segments.ContainsByPredicate([](const FPropertyBindingPathSegment& Segment)
+		{
+			return Segment.GetInstanceStruct() != nullptr;
+		});
+	}
+
+	/** Reset the path to empty. */
+	void Reset()
+	{
+#if WITH_EDITORONLY_DATA
+		StructID = FGuid();
+#endif
+		Segments.Reset();
+	}
+
+#if WITH_EDITORONLY_DATA
+	const FGuid& GetStructID() const
+	{
+		return StructID;
+	}
+	
+	void SetStructID(const FGuid NewStructID)
+	{
+		StructID = NewStructID;
+	}
+#endif // WITH_EDITORONLY_DATA
+
+	TConstArrayView<FPropertyBindingPathSegment> GetSegments() const
+	{
+		return Segments;
+	}
+	
+	TArrayView<FPropertyBindingPathSegment> GetMutableSegments()
+	{
+		return Segments;
+	}
+	
+	int32 NumSegments() const
+	{
+		return Segments.Num();
+	}
+	
+	const FPropertyBindingPathSegment& GetSegment(const int32 Index) const
+	{
+		return Segments[Index];
+	}
+
+	/** Adds a path segment to the path. */
+	void AddPathSegment(const FName InName, const int32 InArrayIndex = INDEX_NONE, const UStruct* InInstanceType = nullptr)
+	{
+		Segments.Emplace(InName, InArrayIndex, InInstanceType);
+	}
+
+	/** Adds a path segment to the path. */
+	void AddPathSegment(const FPropertyBindingPathSegment& PathSegment)
+	{
+		Segments.Add(PathSegment);
+	}
+
+	/** Test if paths are equal. */
+	bool operator==(const FPropertyBindingPath& RHS) const;
+
+private:
+#if WITH_EDITORONLY_DATA
+	/** ID of the struct this property path is relative to. */
+	UPROPERTY()
+	FGuid StructID;
+#endif // WITH_EDITORONLY_DATA
+
+	/** Path segments pointing to a specific property on the path. */
+	UPROPERTY()
+	TArray<FPropertyBindingPathSegment> Segments;
+};
\ No newline at end of file
diff --git a/Engine/Plugins/Runtime/PropertyBindingUtils/Source/PropertyBindingUtils/Public/PropertyBindingUtilsModule.h b/Engine/Plugins/Runtime/PropertyBindingUtils/Source/PropertyBindingUtils/Public/PropertyBindingUtilsModule.h
new file mode 100644
--- /dev/null
+++ b/Engine/Plugins/Runtime/PropertyBindingUtils/Source/PropertyBindingUtils/Public/PropertyBindingUtilsModule.h
@@ -0,0 +1,35 @@
+// Copyright Epic Games, Inc. All Rights Reserved.
+
+#pragma  once
+
+#include "Modules/ModuleManager.h"
+
+/**
+* The public interface to this module
+*/
+class IPropertyBindingUtilsModule : public IModuleInterface
+{
+
+public:
+
+	/**
+	* Singleton-like access to this module's interface.  This is just for convenience!
+	* Beware of calling this during the shutdown phase, though.  Your module might have been unloaded already.
+	*
+	* @return Returns singleton instance, loading the module on demand if needed
+	*/
+	static inline IPropertyBindingUtilsModule& Get()
+	{
+		return FModuleManager::LoadModuleChecked<IPropertyBindingUtilsModule>("PropertyBindingUtilsModule");
+	}
+
+	/**
+	* Checks to see if this module is loaded and ready.  It is only valid to call Get() if IsAvailable() returns true.
+	*
+	* @return True if the module is loaded and ready to use
+	*/
+	static inline bool IsAvailable()
+	{
+		return FModuleManager::Get().IsModuleLoaded("PropertyBindingUtilsModule");
+	}
+};
diff --git a/Engine/Plugins/Runtime/PropertyBindingUtils/Source/PropertyBindingUtilsTestSuite/Private/PropertyBindingUtilsTest.cpp b/Engine/Plugins/Runtime/PropertyBindingUtils/Source/PropertyBindingUtilsTestSuite/Private/PropertyBindingUtilsTest.cpp
new file mode 100644
--- /dev/null
+++ b/Engine/Plugins/Runtime/PropertyBindingUtils/Source/PropertyBindingUtilsTestSuite/Private/PropertyBindingUtilsTest.cpp
@@ -0,0 +1,336 @@
+// Copyright Epic Games, Inc. All Rights Reserved.
+
+#include "PropertyBindingUtilsTest.h"
+#include "PropertyBindingPath.h"
+#include "AITestsCommon.h"
+#include "Engine/World.h"
+
+#include UE_INLINE_GENERATED_CPP_BY_NAME(PropertyBindingUtilsTest)
+
+#define LOCTEXT_NAMESPACE "AITestSuite_PropertyBindingUtilsTest"
+
+UE_DISABLE_OPTIMIZATION_SHIP
+
+
+struct FPropertyBindingUtilsTest_PropertyPathOffset : FAITestBase
+{
+	virtual bool InstantTest() override
+	{
+		FPropertyBindingPath Path;
+		const bool bParseResult = Path.FromString(TEXT("StructB.B"));
+
+		AITEST_TRUE("Parsing path should succeeed", bParseResult);
+		AITEST_EQUAL("Should have 2 path segments", Path.NumSegments(), 2);
+
+		FString ResolveErrors;
+		TArray<FPropertyBindingPathIndirection> Indirections;
+		const bool bResolveResult = Path.ResolveIndirections(FPropertyBindingUtilsTest_PropertyStruct::StaticStruct(), Indirections, &ResolveErrors);
+
+		AITEST_TRUE("Resolve path should succeeed", bResolveResult);
+		AITEST_EQUAL("Should have no resolve errors", ResolveErrors.Len(), 0);
+		
+		AITEST_EQUAL("Should have 2 indirections", Indirections.Num(), 2);
+		AITEST_EQUAL("Indirection 0 should be Offset type", Indirections[0].GetAccessType(), EPropertyBindingAccessType::Offset);
+		AITEST_EQUAL("Indirection 1 should be Offset type", Indirections[1].GetAccessType(), EPropertyBindingAccessType::Offset);
+
+		return true;
+	}
+};
+IMPLEMENT_AI_INSTANT_TEST(FPropertyBindingUtilsTest_PropertyPathOffset, "System.PropertyBindingUtils.Offset");
+
+struct FPropertyBindingUtilsTest_PropertyPathParseFail : FAITestBase
+{
+	virtual bool InstantTest() override
+	{
+		{
+			FPropertyBindingPath Path;
+			const bool bParseResult = Path.FromString(TEXT("")); // empty is valid.
+			AITEST_TRUE("Parsing path should succeed", bParseResult);
+		}
+
+		{
+			FPropertyBindingPath Path;
+			const bool bParseResult = Path.FromString(TEXT("StructB.[0]B"));
+			AITEST_FALSE("Parsing path should fail", bParseResult);
+		}
+
+		{
+			FPropertyBindingPath Path;
+			const bool bParseResult = Path.FromString(TEXT("StructB..NoThere"));
+			AITEST_FALSE("Parsing path should fail", bParseResult);
+		}
+
+		{
+			FPropertyBindingPath Path;
+			const bool bParseResult = Path.FromString(TEXT("."));
+			AITEST_FALSE("Parsing path should fail", bParseResult);
+		}
+
+		{
+			FPropertyBindingPath Path;
+			const bool bParseResult = Path.FromString(TEXT("StructB..B"));
+			AITEST_FALSE("Parsing path should fail", bParseResult);
+		}
+
+		return true;
+	}
+};
+IMPLEMENT_AI_INSTANT_TEST(FPropertyBindingUtilsTest_PropertyPathParseFail, "System.PropertyBindingUtils.ParseFail");
+
+struct FPropertyBindingUtilsTest_PropertyPathOffsetFail : FAITestBase
+{
+	virtual bool InstantTest() override
+	{
+		FPropertyBindingPath Path;
+		const bool bParseResult = Path.FromString(TEXT("StructB.Q"));
+
+		AITEST_TRUE("Parsing path should succeeed", bParseResult);
+		AITEST_EQUAL("Should have 2 path segments", Path.NumSegments(), 2);
+
+		FString ResolveErrors;
+		TArray<FPropertyBindingPathIndirection> Indirections;
+		const bool bResolveResult = Path.ResolveIndirections(FPropertyBindingUtilsTest_PropertyStruct::StaticStruct(), Indirections, &ResolveErrors);
+
+		AITEST_FALSE("Resolve path should not succeeed", bResolveResult);
+		AITEST_NOT_EQUAL("Should have errors", ResolveErrors.Len(), 0);
+		
+		AITEST_EQUAL("Should have 0 indirections", Indirections.Num(), 0);
+
+		return true;
+	}
+};
+IMPLEMENT_AI_INSTANT_TEST(FPropertyBindingUtilsTest_PropertyPathOffsetFail, "System.PropertyBindingUtils.OffsetFail");
+
+struct FPropertyBindingUtilsTest_PropertyPathObject : FAITestBase
+{
+	virtual bool InstantTest() override
+	{
+		FPropertyBindingPath Path;
+		const bool bParseResult = Path.FromString(TEXT("InstancedObject.A"));
+
+		AITEST_TRUE("Parsing path should succeeed", bParseResult);
+		AITEST_EQUAL("Should have 2 path segments", Path.NumSegments(), 2);
+
+		UPropertyBindingUtilsTest_PropertyObject* Object = NewObject<UPropertyBindingUtilsTest_PropertyObject>();
+		Object->InstancedObject = NewObject<UPropertyBindingUtilsTest_PropertyObjectInstanced>();
+		
+		const bool bUpdateResult = Path.UpdateSegmentsFromValue(FPropertyBindingDataView(Object));
+
+		AITEST_TRUE("Update instance types should succeeed", bUpdateResult);
+		AITEST_TRUE("Path segment 0 instance type should be UPropertyBindingUtilsTest_PropertyObjectInstanced", Path.GetSegment(0).GetInstanceStruct() == UPropertyBindingUtilsTest_PropertyObjectInstanced::StaticClass());
+		AITEST_TRUE("Path segment 1 instance type should be nullptr", Path.GetSegment(1).GetInstanceStruct() == nullptr);
+
+		return true;
+	}
+};
+IMPLEMENT_AI_INSTANT_TEST(FPropertyBindingUtilsTest_PropertyPathObject, "System.PropertyBindingUtils.Object");
+
+struct FPropertyBindingUtilsTest_PropertyPathWrongObject : FAITestBase
+{
+	virtual bool InstantTest() override
+	{
+		FPropertyBindingPath Path;
+		const bool bParseResult = Path.FromString(TEXT("InstancedObject.B"));
+
+		AITEST_TRUE("Parsing path should succeeed", bParseResult);
+		AITEST_EQUAL("Should have 2 path segments", Path.NumSegments(), 2);
+
+		UPropertyBindingUtilsTest_PropertyObject* Object = NewObject<UPropertyBindingUtilsTest_PropertyObject>();
+
+		Object->InstancedObject = NewObject<UPropertyBindingUtilsTest_PropertyObjectInstancedWithB>();
+		{
+			FString ResolveErrors;
+			TArray<FPropertyBindingPathIndirection> Indirections;
+			const bool bResolveResult = Path.ResolveIndirectionsWithValue(FPropertyBindingDataView(Object), Indirections, &ResolveErrors);
+
+			AITEST_TRUE("Resolve path should succeeed", bResolveResult);
+			AITEST_EQUAL("Should have 2 indirections", Indirections.Num(), 2);
+			AITEST_TRUE("Object ", Indirections[0].GetAccessType() == EPropertyBindingAccessType::ObjectInstance);
+			AITEST_TRUE("Object ", Indirections[0].GetContainerStruct() == Object->GetClass());
+			AITEST_TRUE("Object ", Indirections[0].GetInstanceStruct() == UPropertyBindingUtilsTest_PropertyObjectInstancedWithB::StaticClass());
+			AITEST_EQUAL("Should not have error", ResolveErrors.Len(), 0);
+		}
+
+		Object->InstancedObject = NewObject<UPropertyBindingUtilsTest_PropertyObjectInstanced>();
+		{
+			FString ResolveErrors;
+			TArray<FPropertyBindingPathIndirection> Indirections;
+			const bool bResolveResult = Path.ResolveIndirectionsWithValue(FPropertyBindingDataView(Object), Indirections, &ResolveErrors);
+
+			AITEST_FALSE("Resolve path should fail", bResolveResult);
+			AITEST_EQUAL("Should have 0 indirections", Indirections.Num(), 0);
+			AITEST_NOT_EQUAL("Should have error", ResolveErrors.Len(), 0);
+		}
+		
+		return true;
+	}
+};
+IMPLEMENT_AI_INSTANT_TEST(FPropertyBindingUtilsTest_PropertyPathWrongObject, "System.PropertyBindingUtils.WrongObject");
+
+struct FPropertyBindingUtilsTest_PropertyPathArray : FAITestBase
+{
+	virtual bool InstantTest() override
+	{
+		FPropertyBindingPath Path;
+		const bool bParseResult = Path.FromString(TEXT("ArrayOfInts[1]"));
+
+		AITEST_TRUE("Parsing path should succeeed", bParseResult);
+		AITEST_EQUAL("Should have 1 path segments", Path.NumSegments(), 1);
+
+		UPropertyBindingUtilsTest_PropertyObject* Object = NewObject<UPropertyBindingUtilsTest_PropertyObject>();
+		Object->ArrayOfInts.Add(42);
+		Object->ArrayOfInts.Add(123);
+
+		FString ResolveErrors;
+		TArray<FPropertyBindingPathIndirection> Indirections;
+		const bool bResolveResult = Path.ResolveIndirectionsWithValue(FPropertyBindingDataView(Object), Indirections, &ResolveErrors);
+
+		AITEST_TRUE("Resolve path should succeeed", bResolveResult);
+		AITEST_EQUAL("Should have no resolve errors", ResolveErrors.Len(), 0);
+		AITEST_EQUAL("Should have 2 indirections", Indirections.Num(), 2);
+		AITEST_EQUAL("Indirection 0 should be IndexArray type", Indirections[0].GetAccessType(), EPropertyBindingAccessType::IndexArray);
+		AITEST_EQUAL("Indirection 1 should be Offset type", Indirections[1].GetAccessType(), EPropertyBindingAccessType::Offset);
+
+		const int32 Value = *reinterpret_cast<const int32*>(Indirections[1].GetPropertyAddress());
+		AITEST_EQUAL("Value should be 123", Value, 123);
+		
+		return true;
+	}
+};
+IMPLEMENT_AI_INSTANT_TEST(FPropertyBindingUtilsTest_PropertyPathArray, "System.PropertyBindingUtils.Array");
+
+struct FPropertyBindingUtilsTest_PropertyPathArrayInvalidIndex : FAITestBase
+{
+	virtual bool InstantTest() override
+	{
+		FPropertyBindingPath Path;
+		const bool bParseResult = Path.FromString(TEXT("ArrayOfInts[123]"));
+
+		AITEST_TRUE("Parsing path should succeeed", bParseResult);
+		AITEST_EQUAL("Should have 1 path segments", Path.NumSegments(), 1);
+
+		UPropertyBindingUtilsTest_PropertyObject* Object = NewObject<UPropertyBindingUtilsTest_PropertyObject>();
+		Object->ArrayOfInts.Add(42);
+		Object->ArrayOfInts.Add(123);
+
+		FString ResolveErrors;
+		TArray<FPropertyBindingPathIndirection> Indirections;
+		const bool bResolveResult = Path.ResolveIndirectionsWithValue(FPropertyBindingDataView(Object), Indirections, &ResolveErrors);
+
+		AITEST_FALSE("Resolve path should fail", bResolveResult);
+		
+		return true;
+	}
+};
+IMPLEMENT_AI_INSTANT_TEST(FPropertyBindingUtilsTest_PropertyPathArrayInvalidIndex, "System.PropertyBindingUtils.ArrayInvalidIndex");
+
+struct FPropertyBindingUtilsTest_PropertyPathArrayOfStructs : FAITestBase
+{
+	virtual bool InstantTest() override
+	{
+		FPropertyBindingPath Path1;
+		Path1.FromString(TEXT("ArrayOfStruct[0].B"));
+
+		FPropertyBindingPath Path2;
+		Path2.FromString(TEXT("ArrayOfStruct[2].StructB.B"));
+
+		UPropertyBindingUtilsTest_PropertyObject* Object = NewObject<UPropertyBindingUtilsTest_PropertyObject>();
+		Object->ArrayOfStruct.AddDefaulted_GetRef().B = 3;
+		Object->ArrayOfStruct.AddDefaulted();
+		Object->ArrayOfStruct.AddDefaulted_GetRef().StructB.B = 42;
+
+		{
+			FString ResolveErrors;
+			TArray<FPropertyBindingPathIndirection> Indirections;
+			const bool bResolveResult = Path1.ResolveIndirectionsWithValue(FPropertyBindingDataView(Object), Indirections, &ResolveErrors);
+
+			AITEST_TRUE("Resolve path1 should succeeed", bResolveResult);
+			AITEST_EQUAL("Should have no resolve errors", ResolveErrors.Len(), 0);
+			AITEST_EQUAL("Should have 3 indirections", Indirections.Num(), 3);
+			AITEST_EQUAL("Indirection 0 should be ArrayIndex type", Indirections[0].GetAccessType(), EPropertyBindingAccessType::IndexArray);
+			AITEST_EQUAL("Indirection 1 should be Offset type", Indirections[1].GetAccessType(), EPropertyBindingAccessType::Offset);
+			AITEST_EQUAL("Indirection 2 should be Offset type", Indirections[2].GetAccessType(), EPropertyBindingAccessType::Offset);
+
+			const int32 Value = *reinterpret_cast<const int32*>(Indirections[2].GetPropertyAddress());
+			AITEST_EQUAL("Value should be 3", Value, 3);
+		}
+
+		{
+			FString ResolveErrors;
+			TArray<FPropertyBindingPathIndirection> Indirections;
+			const bool bResolveResult = Path2.ResolveIndirectionsWithValue(FPropertyBindingDataView(Object), Indirections, &ResolveErrors);
+
+			AITEST_TRUE("Resolve path2 should succeeed", bResolveResult);
+			AITEST_EQUAL("Should have no resolve errors", ResolveErrors.Len(), 0);
+			AITEST_EQUAL("Should have 4 indirections", Indirections.Num(), 4);
+			AITEST_EQUAL("Indirection 0 should be ArrayIndex type", Indirections[0].GetAccessType(), EPropertyBindingAccessType::IndexArray);
+			AITEST_EQUAL("Indirection 1 should be Offset type", Indirections[1].GetAccessType(), EPropertyBindingAccessType::Offset);
+			AITEST_EQUAL("Indirection 2 should be Offset type", Indirections[2].GetAccessType(), EPropertyBindingAccessType::Offset);
+			AITEST_EQUAL("Indirection 3 should be Offset type", Indirections[3].GetAccessType(), EPropertyBindingAccessType::Offset);
+
+			const int32 Value = *reinterpret_cast<const int32*>(Indirections[3].GetPropertyAddress());
+			AITEST_EQUAL("Value should be 42", Value, 42);
+		}
+		return true;
+	}
+};
+IMPLEMENT_AI_INSTANT_TEST(FPropertyBindingUtilsTest_PropertyPathArrayOfStructs, "System.PropertyBindingUtils.ArrayOfStructs");
+
+struct FPropertyBindingUtilsTest_PropertyPathArrayOfInstancedObjects : FAITestBase
+{
+	virtual bool InstantTest() override
+	{
+		FPropertyBindingPath Path;
+		Path.FromString(TEXT("ArrayOfInstancedStructs[0].B"));
+
+		FPropertyBindingUtilsTest_PropertyStruct Struct;
+		Struct.B = 123;
+		
+		UPropertyBindingUtilsTest_PropertyObject* Object = NewObject<UPropertyBindingUtilsTest_PropertyObject>();
+		Object->ArrayOfInstancedStructs.Emplace(FConstStructView::Make(Struct));
+
+		const bool bUpdateResult = Path.UpdateSegmentsFromValue(FPropertyBindingDataView(Object));
+		AITEST_TRUE("Update instance types should succeeed", bUpdateResult);
+		AITEST_EQUAL("Should have 2 path segments", Path.NumSegments(), 2);
+		AITEST_TRUE("Path segment 0 instance type should be FPropertyBindingUtilsTest_PropertyStruct", Path.GetSegment(0).GetInstanceStruct() == FPropertyBindingUtilsTest_PropertyStruct::StaticStruct());
+		AITEST_TRUE("Path segment 1 instance type should be nullptr", Path.GetSegment(1).GetInstanceStruct() == nullptr);
+
+		{
+			FString ResolveErrors;
+			TArray<FPropertyBindingPathIndirection> Indirections;
+			const bool bResolveResult = Path.ResolveIndirections(UPropertyBindingUtilsTest_PropertyObject::StaticClass(), Indirections, &ResolveErrors);
+
+			AITEST_TRUE("Resolve path should succeeed", bResolveResult);
+			AITEST_EQUAL("Should have no resolve errors", ResolveErrors.Len(), 0);
+			AITEST_EQUAL("Should have 3 indirections", Indirections.Num(), 3);
+			AITEST_EQUAL("Indirection 0 should be ArrayIndex type", Indirections[0].GetAccessType(), EPropertyBindingAccessType::IndexArray);
+			AITEST_EQUAL("Indirection 1 should be StructInstance type", Indirections[1].GetAccessType(), EPropertyBindingAccessType::StructInstance);
+			AITEST_EQUAL("Indirection 2 should be Offset type", Indirections[2].GetAccessType(), EPropertyBindingAccessType::Offset);
+		}
+
+		{
+			FString ResolveErrors;
+			TArray<FPropertyBindingPathIndirection> Indirections;
+			const bool bResolveResult = Path.ResolveIndirectionsWithValue(FPropertyBindingDataView(Object), Indirections, &ResolveErrors);
+
+			AITEST_TRUE("Resolve path should succeeed", bResolveResult);
+			AITEST_EQUAL("Should have no resolve errors", ResolveErrors.Len(), 0);
+			AITEST_EQUAL("Should have 3 indirections", Indirections.Num(), 3);
+			AITEST_EQUAL("Indirection 0 should be ArrayIndex type", Indirections[0].GetAccessType(), EPropertyBindingAccessType::IndexArray);
+			AITEST_EQUAL("Indirection 1 should be StructInstance type", Indirections[1].GetAccessType(), EPropertyBindingAccessType::StructInstance);
+			AITEST_EQUAL("Indirection 2 should be Offset type", Indirections[2].GetAccessType(), EPropertyBindingAccessType::Offset);
+
+			const int32 Value = *reinterpret_cast<const int32*>(Indirections[2].GetPropertyAddress());
+			AITEST_EQUAL("Value should be 123", Value, 123);
+		}
+		
+		return true;
+	}
+};
+IMPLEMENT_AI_INSTANT_TEST(FPropertyBindingUtilsTest_PropertyPathArrayOfInstancedObjects, "System.PropertyBindingUtils.ArrayOfInstancedObjects");
+
+
+UE_ENABLE_OPTIMIZATION_SHIP
+
+#undef LOCTEXT_NAMESPACE
+
diff --git a/Engine/Plugins/Runtime/PropertyBindingUtils/Source/PropertyBindingUtilsTestSuite/Private/PropertyBindingUtilsTest.h b/Engine/Plugins/Runtime/PropertyBindingUtils/Source/PropertyBindingUtilsTestSuite/Private/PropertyBindingUtilsTest.h
new file mode 100644
--- /dev/null
+++ b/Engine/Plugins/Runtime/PropertyBindingUtils/Source/PropertyBindingUtilsTestSuite/Private/PropertyBindingUtilsTest.h
@@ -0,0 +1,82 @@
+// Copyright Epic Games, Inc. All Rights Reserved.
+
+#pragma once
+
+#include "InstancedStruct.h"
+#include "PropertyBindingUtilsTest.generated.h"
+
+
+USTRUCT()
+struct FPropertyBindingUtilsTest_PropertyStructB
+{
+	GENERATED_BODY()
+
+	UPROPERTY(EditAnywhere, Category = "")
+	int32 B = 0;
+};
+
+USTRUCT()
+struct FPropertyBindingUtilsTest_PropertyStruct
+{
+	GENERATED_BODY()
+
+	UPROPERTY(EditAnywhere, Category = "")
+	int32 A = 0;
+
+	UPROPERTY(EditAnywhere, Category = "")
+	int32 B = 0;
+
+	UPROPERTY(EditAnywhere, Category = "")
+	FPropertyBindingUtilsTest_PropertyStructB StructB;
+};
+
+UCLASS(HideDropdown)
+class UPropertyBindingUtilsTest_PropertyObjectInstanced : public UObject
+{
+	GENERATED_BODY()
+public:
+
+	UPROPERTY(EditAnywhere, Category = "")
+	int32 A = 0;
+
+	UPROPERTY(EditAnywhere, Category = "")
+	FInstancedStruct InstancedStruct;
+};
+
+UCLASS(HideDropdown)
+class UPropertyBindingUtilsTest_PropertyObjectInstancedWithB : public UPropertyBindingUtilsTest_PropertyObjectInstanced
+{
+	GENERATED_BODY()
+public:
+
+	UPROPERTY(EditAnywhere, Category = "")
+	int32 B = 0;
+};
+
+UCLASS(HideDropdown)
+class UPropertyBindingUtilsTest_PropertyObject : public UObject
+{
+	GENERATED_BODY()
+public:
+	
+	UPROPERTY(EditAnywhere, Instanced, Category = "")
+	TObjectPtr<UPropertyBindingUtilsTest_PropertyObjectInstanced> InstancedObject;
+
+	UPROPERTY(EditAnywhere, Instanced, Category = "")
+	TArray<TObjectPtr<UPropertyBindingUtilsTest_PropertyObjectInstanced>> ArrayOfInstancedObjects;
+
+	UPROPERTY(EditAnywhere, Category = "")
+	TArray<int32> ArrayOfInts;
+
+	UPROPERTY(EditAnywhere, Category = "")
+	FInstancedStruct InstancedStruct;
+
+	UPROPERTY(EditAnywhere, Category = "")
+	TArray<FInstancedStruct> ArrayOfInstancedStructs;
+
+	UPROPERTY(EditAnywhere, Category = "")
+	FPropertyBindingUtilsTest_PropertyStruct Struct;
+
+	UPROPERTY(EditAnywhere, Category = "")
+	TArray<FPropertyBindingUtilsTest_PropertyStruct> ArrayOfStruct;
+};
\ No newline at end of file
diff --git a/Engine/Plugins/Runtime/PropertyBindingUtils/Source/PropertyBindingUtilsTestSuite/Private/PropertyBindingUtilsTestSuiteModule.cpp b/Engine/Plugins/Runtime/PropertyBindingUtils/Source/PropertyBindingUtilsTestSuite/Private/PropertyBindingUtilsTestSuiteModule.cpp
new file mode 100644
--- /dev/null
+++ b/Engine/Plugins/Runtime/PropertyBindingUtils/Source/PropertyBindingUtilsTestSuite/Private/PropertyBindingUtilsTestSuiteModule.cpp
@@ -0,0 +1,13 @@
+// Copyright Epic Games, Inc. All Rights Reserved.
+
+#include "PropertyBindingUtilsTestSuiteModule.h"
+
+#define LOCTEXT_NAMESPACE "PropertyBindingUtilsTestSuite"
+
+class FPropertyBindingUtilsTestSuiteModule : public IPropertyBindingUtilsTestSuiteModule
+{
+};
+
+IMPLEMENT_MODULE(FPropertyBindingUtilsTestSuiteModule, PropertyBindingUtilsTestSuite)
+
+#undef LOCTEXT_NAMESPACE
diff --git a/Engine/Plugins/Runtime/PropertyBindingUtils/Source/PropertyBindingUtilsTestSuite/PropertyBindingUtilsTestSuite.Build.cs b/Engine/Plugins/Runtime/PropertyBindingUtils/Source/PropertyBindingUtilsTestSuite/PropertyBindingUtilsTestSuite.Build.cs
new file mode 100644
--- /dev/null
+++ b/Engine/Plugins/Runtime/PropertyBindingUtils/Source/PropertyBindingUtilsTestSuite/PropertyBindingUtilsTestSuite.Build.cs
@@ -0,0 +1,29 @@
+// Copyright Epic Games, Inc. All Rights Reserved.
+
+using UnrealBuildTool;
+
+namespace UnrealBuildTool.Rules
+{
+	public class PropertyBindingUtilsTestSuite : ModuleRules
+	{
+		public PropertyBindingUtilsTestSuite(ReadOnlyTargetRules Target) : base(Target)
+		{
+			PublicDependencyModuleNames.AddRange(
+				new string[] {
+					"Core",
+					"CoreUObject",
+					"Engine",
+					"PropertyBindingUtils",
+					"AITestSuite",
+					"StructUtils",
+				}
+			);
+
+			if (Target.bBuildEditor == true)
+			{
+				PrivateDependencyModuleNames.Add("EditorFramework");
+				PrivateDependencyModuleNames.Add("UnrealEd"); 
+			}
+		}
+	}
+}
\ No newline at end of file
diff --git a/Engine/Plugins/Runtime/PropertyBindingUtils/Source/PropertyBindingUtilsTestSuite/Public/PropertyBindingUtilsTestSuiteModule.h b/Engine/Plugins/Runtime/PropertyBindingUtils/Source/PropertyBindingUtilsTestSuite/Public/PropertyBindingUtilsTestSuiteModule.h
new file mode 100644
--- /dev/null
+++ b/Engine/Plugins/Runtime/PropertyBindingUtils/Source/PropertyBindingUtilsTestSuite/Public/PropertyBindingUtilsTestSuiteModule.h
@@ -0,0 +1,35 @@
+// Copyright Epic Games, Inc. All Rights Reserved.
+
+#pragma  once
+
+#include "Modules/ModuleManager.h"
+
+/**
+* The public interface to this module
+*/
+class IPropertyBindingUtilsTestSuiteModule : public IModuleInterface
+{
+
+public:
+
+	/**
+	* Singleton-like access to this module's interface.  This is just for convenience!
+	* Beware of calling this during the shutdown phase, though.  Your module might have been unloaded already.
+	*
+	* @return Returns singleton instance, loading the module on demand if needed
+	*/
+	static inline IPropertyBindingUtilsTestSuiteModule& Get()
+	{
+		return FModuleManager::LoadModuleChecked<IPropertyBindingUtilsTestSuiteModule>("PropertyBindingUtilsTestSuiteModule");
+	}
+
+	/**
+	* Checks to see if this module is loaded and ready.  It is only valid to call Get() if IsAvailable() returns true.
+	*
+	* @return True if the module is loaded and ready to use
+	*/
+	static inline bool IsAvailable()
+	{
+		return FModuleManager::Get().IsModuleLoaded("PropertyBindingUtilsTestSuiteModule");
+	}
+};
