[PropertyBindingUtils]
* Moved dll export to functions instead of on type to solve some linking errors when building non-unity non-pch

diff --git a/Engine/Plugins/Runtime/PropertyBindingUtils/Source/PropertyBindingUtils/Public/PropertyBindingPath.h b/Engine/Plugins/Runtime/PropertyBindingUtils/Source/PropertyBindingUtils/Public/PropertyBindingPath.h
--- a/Engine/Plugins/Runtime/PropertyBindingUtils/Source/PropertyBindingUtils/Public/PropertyBindingPath.h
+++ b/Engine/Plugins/Runtime/PropertyBindingUtils/Source/PropertyBindingUtils/Public/PropertyBindingPath.h
@@ -76,7 +76,7 @@ enum class EPropertyBindingAccessType : uint8
  * When creating from a generic UStruct, the data view expects a valid type (UStruct) when a valid memory pointer is provided.
  * It that case it is also fine to provide null for both and the constructed data view will be considered as invalid.
  */
-struct PROPERTYBINDINGUTILS_API FPropertyBindingDataView
+struct FPropertyBindingDataView
 {
 	FPropertyBindingDataView() = default;
 
@@ -261,7 +261,7 @@ protected:
  * Returned by FPropertyBindingPath::ResolveIndirections() and FPropertyBindingPath::ResolveIndirectionsWithValue().
  * Generally there's one indirection per FProperty. Containers have one path segment but two indirection (container + inner type).
  */
-struct PROPERTYBINDINGUTILS_API FPropertyBindingPathIndirection
+struct FPropertyBindingPathIndirection
 {
 	FPropertyBindingPathIndirection() = default;
 	explicit FPropertyBindingPathIndirection(const UStruct* InContainerStruct)
@@ -369,7 +369,7 @@ private:
 
 /** Struct describing a path segment in FPropertyBindingPath. */
 USTRUCT()
-struct PROPERTYBINDINGUTILS_API FPropertyBindingPathSegment
+struct FPropertyBindingPathSegment
 {
 	GENERATED_BODY()
 
@@ -464,7 +464,7 @@ private:
  * You may also take a look at: FCachedPropertyPath, TFieldPath<>, FPropertyPath.
  */
 USTRUCT()
-struct PROPERTYBINDINGUTILS_API FPropertyBindingPath
+struct FPropertyBindingPath
 {
 	GENERATED_BODY()
 
@@ -494,7 +494,7 @@ struct PROPERTYBINDINGUTILS_API FPropertyBindingPath
 	 * @param InPath Path string to parse
 	 * @return true if path was parsed successfully.
 	 */
-	bool FromString(const FStringView InPath);
+	PROPERTYBINDINGUTILS_API bool FromString(const FStringView InPath);
 
 	/**
 	 * Returns the property path as a one string. Highlight allows to decorate a specific segment.
@@ -504,7 +504,7 @@ struct PROPERTYBINDINGUTILS_API FPropertyBindingPath
 	 * @param bOutputInstances if true, the instance struct types will be output. 
 	 * @param FirstSegment Index of the first path segment to be stringified.
 	 */
-	FString ToString(const int32 HighlightedSegment = INDEX_NONE, const TCHAR* HighlightPrefix = nullptr, const TCHAR* HighlightPostfix = nullptr, const bool bOutputInstances = false, const int32 FirstSegment = 0) const;
+	PROPERTYBINDINGUTILS_API FString ToString(const int32 HighlightedSegment = INDEX_NONE, const TCHAR* HighlightPrefix = nullptr, const TCHAR* HighlightPostfix = nullptr, const bool bOutputInstances = false, const int32 FirstSegment = 0) const;
 
 	/**
 	 * Resolves the property path against base struct type. The path is assumed to be relative to the BaseStruct.
@@ -514,7 +514,7 @@ struct PROPERTYBINDINGUTILS_API FPropertyBindingPath
 	 * @param bHandleRedirects If true, the method will try to resolve missing properties using core redirects, and properties on Blueprint and User Defined Structs by ID. Available only in editor builds!
 	 * @return true of path could be resolved against the base value, and instance types were updated.
 	 */
-	bool ResolveIndirections(const UStruct* BaseStruct, TArray<FPropertyBindingPathIndirection>& OutIndirections, FString* OutError = nullptr, bool bHandleRedirects = false) const;
+	PROPERTYBINDINGUTILS_API bool ResolveIndirections(const UStruct* BaseStruct, TArray<FPropertyBindingPathIndirection>& OutIndirections, FString* OutError = nullptr, bool bHandleRedirects = false) const;
 	
 	/**
 	 * Resolves the property path against base value. The path is assumed to be relative to the BaseValueView.
@@ -524,7 +524,7 @@ struct PROPERTYBINDINGUTILS_API FPropertyBindingPath
 	 * @param bHandleRedirects If true, the method will try to resolve missing properties using core redirects, and properties on Blueprint and User Defined Structs by ID. Available only in editor builds!
 	 * @return true of path could be resolved against the base value, and instance types were updated.
 	 */
-	bool ResolveIndirectionsWithValue(const FPropertyBindingDataView BaseValueView, TArray<FPropertyBindingPathIndirection>& OutIndirections, FString* OutError = nullptr, bool bHandleRedirects = false) const;
+	PROPERTYBINDINGUTILS_API bool ResolveIndirectionsWithValue(const FPropertyBindingDataView BaseValueView, TArray<FPropertyBindingPathIndirection>& OutIndirections, FString* OutError = nullptr, bool bHandleRedirects = false) const;
 
 	/**
 	 * Updates property segments from base struct type. The path is expected to be relative to the BaseStruct.
@@ -533,7 +533,7 @@ struct PROPERTYBINDINGUTILS_API FPropertyBindingPath
 	 * @param OutError Optional, pointer to string where error will be logged if update fails.
 	 * @return true of path could be resolved against the base value, and instance types were updated.
 	 */
-	bool UpdateSegments(const UStruct* BaseStruct, FString* OutError = nullptr);
+	PROPERTYBINDINGUTILS_API bool UpdateSegments(const UStruct* BaseStruct, FString* OutError = nullptr);
 
 	/**
 	 * Updates property segments from base value. The path is expected to be relative to the base value.
@@ -543,7 +543,7 @@ struct PROPERTYBINDINGUTILS_API FPropertyBindingPath
 	 * @param OutError Optional, pointer to string where error will be logged if update fails.
 	 * @return true of path could be resolved against the base value, and instance types were updated.
 	 */
-	bool UpdateSegmentsFromValue(const FPropertyBindingDataView BaseValueView, FString* OutError = nullptr);
+	PROPERTYBINDINGUTILS_API bool UpdateSegmentsFromValue(const FPropertyBindingDataView BaseValueView, FString* OutError = nullptr);
 
 	/** @return true if the path is empty. In that case the path points to the struct. */
 	bool IsPathEmpty() const
@@ -614,13 +614,13 @@ struct PROPERTYBINDINGUTILS_API FPropertyBindingPath
 	}
 
 	/** Test if paths are equal. */
-	bool operator==(const FPropertyBindingPath& RHS) const;
+	PROPERTYBINDINGUTILS_API bool operator==(const FPropertyBindingPath& RHS) const;
 
 	/**
 	 * Test if this paths includes the provided path.
 	 * A path includes another one when they are == but this path can be longer.
 	 */
-	bool Includes(const FPropertyBindingPath& Other) const;
+	PROPERTYBINDINGUTILS_API bool Includes(const FPropertyBindingPath& Other) const;
 
 private:
 #if WITH_EDITORONLY_DATA
