[PropertyBindingsUtils]
- added new base types to shared common functionalities in property bindings:
FPropertyBindingExtension
FPropertyBindingBindableStructDescriptor
FPropertyBindingBinding
FPropertyBindingBindingCollection
IPropertyBindingBindingCollectionOwner
UE::PropertyBinding
{
EPropertyBindingPropertyAccessType
EPropertyCompatibility
FPropertyCreationDescriptor
FCachedBindingData
FPropertyInfoOverride
FMenuSectionHelper
}
Functionality was extracted from StateTree and SmartObject frameworks and now extended by those.

diff --git a/Engine/Plugins/Experimental/Avalanche/Source/AvalancheTransitionEditor/AvalancheTransitionEditor.Build.cs b/Engine/Plugins/Experimental/Avalanche/Source/AvalancheTransitionEditor/AvalancheTransitionEditor.Build.cs
--- a/Engine/Plugins/Experimental/Avalanche/Source/AvalancheTransitionEditor/AvalancheTransitionEditor.Build.cs
+++ b/Engine/Plugins/Experimental/Avalanche/Source/AvalancheTransitionEditor/AvalancheTransitionEditor.Build.cs
@@ -32,6 +32,7 @@ public class AvalancheTransitionEditor : ModuleRules
                 "InputCore",
                 "MessageLog",
                 "Projects",
+                "PropertyBindingUtils",
                 "PropertyEditor",
                 "Slate",
                 "SlateCore",
diff --git a/Engine/Plugins/Experimental/Avalanche/Source/AvalancheTransitionEditor/Private/AvaTransitionEditorUtils.cpp b/Engine/Plugins/Experimental/Avalanche/Source/AvalancheTransitionEditor/Private/AvaTransitionEditorUtils.cpp
--- a/Engine/Plugins/Experimental/Avalanche/Source/AvalancheTransitionEditor/Private/AvaTransitionEditorUtils.cpp
+++ b/Engine/Plugins/Experimental/Avalanche/Source/AvalancheTransitionEditor/Private/AvaTransitionEditorUtils.cpp
@@ -174,9 +174,9 @@ void ValidateTree(UAvaTransitionTree& InTransitionTree)
 
 	// Remove unused Bindings
 	{
-		TMap<FGuid, const FStateTreeDataView> AllStructValues;
+		TMap<FGuid, const FPropertyBindingDataView> AllStructValues;
 		EditorData->GetAllStructValues(AllStructValues);
-		EditorData->GetPropertyEditorBindings()->RemoveUnusedBindings(AllStructValues);
+		EditorData->GetPropertyEditorBindings()->RemoveInvalidBindings(AllStructValues);
 	}
 
 	// Validate Linked States
diff --git a/Engine/Plugins/Experimental/Avalanche/Source/AvalancheTransitionEditor/Private/Serialization/AvaTransitionStateSerializer.cpp b/Engine/Plugins/Experimental/Avalanche/Source/AvalancheTransitionEditor/Private/Serialization/AvaTransitionStateSerializer.cpp
--- a/Engine/Plugins/Experimental/Avalanche/Source/AvalancheTransitionEditor/Private/Serialization/AvaTransitionStateSerializer.cpp
+++ b/Engine/Plugins/Experimental/Avalanche/Source/AvalancheTransitionEditor/Private/Serialization/AvaTransitionStateSerializer.cpp
@@ -50,13 +50,14 @@ namespace UE::AvaTransitionEditor::Private
 		InEditorData.VisitStateNodes(*InState,
 			[&InEditorData, &InCopyData](const UStateTreeState*, const FStateTreeBindableStructDesc& InDesc, const FStateTreeDataView)
 			{
-				TArray<const FStateTreePropertyPathBinding*> NodeBindings;
-				InEditorData.GetPropertyEditorBindings()->GetPropertyBindingsFor(InDesc.ID, NodeBindings);
-				for (const FStateTreePropertyPathBinding* Binding : NodeBindings)
+				TArray<const FPropertyBindingBinding*> NodeBindings;
+				InEditorData.GetPropertyEditorBindings()->GetBindingsFor(InDesc.ID, NodeBindings);
+				for (const FPropertyBindingBinding* Binding : NodeBindings)
 				{
-					InCopyData.Bindings.Add(*Binding);
+					const FStateTreePropertyPathBinding* StateTreeBinding = static_cast<const FStateTreePropertyPathBinding*>(Binding);
+					InCopyData.Bindings.Add(*StateTreeBinding);
 				}
-				return EStateTreeVisitor::Continue;				
+				return EStateTreeVisitor::Continue;
 			});
 
 		for (const UStateTreeState* ChildState : InState->Children)
@@ -250,7 +251,7 @@ bool FAvaTransitionStateSerializer::ImportText(const FString& InText
 						SourcePath.SetStructID(*NewSourceID);
 					}
 
-					InEditorData.GetPropertyEditorBindings()->AddPropertyBinding(SourcePath, TargetPath);
+					InEditorData.GetPropertyEditorBindings()->AddBinding(SourcePath, TargetPath);
 				}
 			}
 		}
diff --git a/Engine/Plugins/Runtime/PropertyBindingUtils/Config/BasePropertyBindingUtils.ini b/Engine/Plugins/Runtime/PropertyBindingUtils/Config/BasePropertyBindingUtils.ini
new file mode 100644
--- /dev/null
+++ b/Engine/Plugins/Runtime/PropertyBindingUtils/Config/BasePropertyBindingUtils.ini
@@ -0,0 +1,4 @@
+[CoreRedirects]
+
++StructRedirects=(OldName="BindableStructDesc",NewName="/script/PropertyBindingUtils.PropertyBindingBindableStructDescriptor")
++EnumRedirects=(OldName="EPropertyBindingAccessType",NewName="/Script/PropertyBindingUtils.EPropertyBindingPropertyAccessType")
diff --git a/Engine/Plugins/Runtime/PropertyBindingUtils/PropertyBindingUtils.uplugin b/Engine/Plugins/Runtime/PropertyBindingUtils/PropertyBindingUtils.uplugin
--- a/Engine/Plugins/Runtime/PropertyBindingUtils/PropertyBindingUtils.uplugin
+++ b/Engine/Plugins/Runtime/PropertyBindingUtils/PropertyBindingUtils.uplugin
@@ -21,6 +21,11 @@
 			"Type": "Runtime",
 			"LoadingPhase": "PreDefault"
 		},
+		{
+			"Name": "PropertyBindingUtilsEditor",
+			"Type": "Editor",
+			"LoadingPhase": "Default"
+		},
 		{
 			"Name": "PropertyBindingUtilsTestSuite",
 			"Type": "UncookedOnly",
diff --git a/Engine/Plugins/Runtime/PropertyBindingUtils/Source/PropertyBindingUtils/Private/PropertyBindingBindableStructDescriptor.cpp b/Engine/Plugins/Runtime/PropertyBindingUtils/Source/PropertyBindingUtils/Private/PropertyBindingBindableStructDescriptor.cpp
new file mode 100644
--- /dev/null
+++ b/Engine/Plugins/Runtime/PropertyBindingUtils/Source/PropertyBindingUtils/Private/PropertyBindingBindableStructDescriptor.cpp
@@ -0,0 +1,15 @@
+// Copyright Epic Games, Inc. All Rights Reserved.
+
+#include "PropertyBindingBindableStructDescriptor.h"
+
+
+FString FPropertyBindingBindableStructDescriptor::ToString() const
+{
+	FStringBuilderBase Result;
+
+	Result += TEXT("'");
+	Result += Name.ToString();
+	Result += TEXT("'");
+
+	return Result.ToString();
+}
\ No newline at end of file
diff --git a/Engine/Plugins/Runtime/PropertyBindingUtils/Source/PropertyBindingUtils/Private/PropertyBindingBindingCollection.cpp b/Engine/Plugins/Runtime/PropertyBindingUtils/Source/PropertyBindingUtils/Private/PropertyBindingBindingCollection.cpp
new file mode 100644
--- /dev/null
+++ b/Engine/Plugins/Runtime/PropertyBindingUtils/Source/PropertyBindingUtils/Private/PropertyBindingBindingCollection.cpp
@@ -0,0 +1,221 @@
+// Copyright Epic Games, Inc. All Rights Reserved.
+
+#include "PropertyBindingBindingCollection.h"
+
+#if WITH_EDITOR
+#include "PropertyBindingBinding.h"
+#include "PropertyBindingBindingCollectionOwner.h"
+#include "PropertyBindingDataView.h"
+#include "PropertyBindingPath.h"
+
+void FPropertyBindingBindingCollection::AddBinding(const FPropertyBindingPath& InSourcePath, const FPropertyBindingPath& InTargetPath)
+{
+	RemoveBindings(InTargetPath, ESearchMode::Exact);
+
+	FPropertyBindingBinding* AddedBinding = AddBindingInternal(InSourcePath, InTargetPath);
+
+	// If we have bindings owner, update property path segments to capture property IDs, etc.
+	const IPropertyBindingBindingCollectionOwner* PropertyBindingsOwner = BindingsOwner.GetInterface();
+	if (PropertyBindingsOwner != nullptr && AddedBinding != nullptr)
+	{
+		FPropertyBindingBinding& Binding = *AddedBinding;
+		FPropertyBindingDataView SourceDataView;
+		if (PropertyBindingsOwner->GetBindingDataViewByID(Binding.GetSourcePath().GetStructID(), SourceDataView))
+		{
+			Binding.GetMutableSourcePath().UpdateSegmentsFromValue(SourceDataView);
+		}
+		FPropertyBindingDataView TargetDataView;
+		if (PropertyBindingsOwner->GetBindingDataViewByID(Binding.GetTargetPath().GetStructID(), TargetDataView))
+		{
+			Binding.GetMutableTargetPath().UpdateSegmentsFromValue(TargetDataView);
+		}
+	}
+}
+
+void FPropertyBindingBindingCollection::AddBinding(const FPropertyBindingBinding& InBinding)
+{
+	RemoveBindings(InBinding.GetTargetPath(), ESearchMode::Exact);
+
+	AddBindingInternal(InBinding);
+}
+
+void FPropertyBindingBindingCollection::RemoveBindings(const FPropertyBindingPath& InTargetPath, ESearchMode InSearchMode)
+{
+	if (InSearchMode == ESearchMode::Exact)
+	{
+		RemoveBindingsInternal([&InTargetPath](const FPropertyBindingBinding& Binding)
+			{
+				return Binding.GetTargetPath() == InTargetPath;
+			});
+	}
+	else
+	{
+		RemoveBindingsInternal([&InTargetPath](const FPropertyBindingBinding& Binding)
+			{
+				return Binding.GetTargetPath().Includes(InTargetPath);
+			});
+	}
+}
+
+void FPropertyBindingBindingCollection::CopyBindings(const FGuid InFromStructID, const FGuid InToStructID)
+{
+	// Copy all bindings that target "InFromStructID" and retarget them to "InToStructID".
+	const int32 NumBindingsToVisit = GetNumBindings();
+	int32 NumVisited = 0;
+	VisitBindings([&NumVisited, NumBindingsToVisit, InFromStructID, InToStructID, this](const FPropertyBindingBinding& Binding)
+	{
+		NumVisited++;
+		if (Binding.GetTargetPath().GetStructID() == InFromStructID)
+		{
+			AddBindingInternal(Binding.GetSourcePath(), FPropertyBindingPath(InToStructID, Binding.GetTargetPath().GetSegments()));
+		}
+		return NumVisited == NumBindingsToVisit ? EVisitResult::Break : EVisitResult::Continue;
+	});
+}
+
+bool FPropertyBindingBindingCollection::HasBinding(const FPropertyBindingPath& InTargetPath, ESearchMode RemoveMode) const
+{
+	if (RemoveMode == ESearchMode::Exact)
+	{
+		return HasBindingInternal([&InTargetPath](const FPropertyBindingBinding& Binding)
+			{
+				return Binding.GetTargetPath() == InTargetPath;
+			});
+	}
+	else
+	{
+		return HasBindingInternal([&InTargetPath](const FPropertyBindingBinding& Binding)
+			{
+				return Binding.GetTargetPath().Includes(InTargetPath);
+			});
+	}
+}
+
+const FPropertyBindingBinding* FPropertyBindingBindingCollection::FindBinding(const FPropertyBindingPath& InTargetPath, ESearchMode RemoveMode) const
+{
+	if (RemoveMode == ESearchMode::Exact)
+	{
+		return FindBindingInternal([&InTargetPath](const FPropertyBindingBinding& Binding)
+			{
+				return Binding.GetTargetPath() == InTargetPath;
+			});
+	}
+	else
+	{
+		return FindBindingInternal([&InTargetPath](const FPropertyBindingBinding& Binding)
+			{
+				return Binding.GetTargetPath().Includes(InTargetPath);
+			});
+	}
+}
+
+const FPropertyBindingPath* FPropertyBindingBindingCollection::GetBindingSource(const FPropertyBindingPath& InTargetPath) const
+{
+	const FPropertyBindingBinding* Binding = FindBindingInternal([&InTargetPath](const FPropertyBindingBinding& Binding)
+		{
+			return Binding.GetTargetPath() == InTargetPath;
+		});
+	return Binding ? &Binding->GetSourcePath() : nullptr;
+}
+
+void FPropertyBindingBindingCollection::GetBindingsFor(const FGuid InStructID, TArray<const FPropertyBindingBinding*>& OutBindings) const
+{
+	ForEachBinding([&OutBindings, &InStructID](const FPropertyBindingBinding& Binding)
+	{
+		if (Binding.GetSourcePath().GetStructID().IsValid() && Binding.GetTargetPath().GetStructID() == InStructID)
+		{
+			OutBindings.Add(&Binding);
+		}
+	});
+}
+
+void FPropertyBindingBindingCollection::RemoveInvalidBindings(const TMap<FGuid, const FPropertyBindingDataView>& InValidStructs)
+{
+	RemoveBindingsInternal([InValidStructs](const FPropertyBindingBinding& Binding)
+		{
+			// Remove binding if it's target struct has been removed
+			if (!InValidStructs.Contains(Binding.GetTargetPath().GetStructID()))
+			{
+				// Remove
+				return true;
+			}
+
+			// Target path should always have at least one segment (copy bind directly on a target struct/object). 
+			if (Binding.GetTargetPath().IsPathEmpty())
+			{
+				return true;
+			}
+
+			// Remove binding if path containing instanced indirections (e.g. instance struct or object) cannot be resolved.
+			// TODO: Try to use core redirect to find new name.
+			{
+				const FPropertyBindingDataView* SourceValue = InValidStructs.Find(Binding.GetSourcePath().GetStructID());
+				if (SourceValue && SourceValue->IsValid())
+				{
+					FString Error;
+					TArray<FPropertyBindingPathIndirection> Indirections;
+					if (!Binding.GetSourcePath().ResolveIndirectionsWithValue(*SourceValue, Indirections, &Error))
+					{
+						UE_LOG(LogPropertyBindingUtils, Verbose, TEXT("Removing binding to %s because binding source path cannot be resolved: %s"),
+							*Binding.GetSourcePath().ToString(), *Error); // Error contains the target path.
+
+						// Remove
+						return true;
+					}
+				}
+			}
+			
+			{
+				const FPropertyBindingDataView TargetValue = InValidStructs.FindChecked(Binding.GetTargetPath().GetStructID());
+				FString Error;
+				TArray<FPropertyBindingPathIndirection> Indirections;
+				if (!Binding.GetTargetPath().ResolveIndirectionsWithValue(TargetValue, Indirections, &Error))
+				{
+					UE_LOG(LogPropertyBindingUtils, Verbose, TEXT("Removing binding to %s because binding target path cannot be resolved: %s"),
+						*Binding.GetSourcePath().ToString(), *Error); // Error contains the target path.
+
+					// Remove
+					return true;
+				}
+			}
+
+			return false;
+		});
+}
+
+bool FPropertyBindingBindingCollection::ContainsAnyStruct(const TSet<const UStruct*>& InStructs) const
+{
+	auto PathContainsStruct = [&InStructs](const FPropertyBindingPath& PropertyPath)
+	{
+		for (const FPropertyBindingPathSegment& Segment : PropertyPath.GetSegments())
+		{
+			if (InStructs.Contains(Segment.GetInstanceStruct()))
+			{
+				return true;
+			}
+		}
+		return false;
+	};
+
+	bool bFound = false;
+	ForEachBinding([&PathContainsStruct, &bFound](const FPropertyBindingBinding& PropertyPathBinding)
+	{
+		if (PathContainsStruct(PropertyPathBinding.GetSourcePath()))
+		{
+			bFound = true;
+			return EVisitResult::Break;
+		}
+
+		if (PathContainsStruct(PropertyPathBinding.GetTargetPath()))
+		{
+			bFound = true;
+			return EVisitResult::Break;
+		}
+
+		return EVisitResult::Continue;
+	});
+	
+	return false;
+}
+
+#endif //WITH_EDITOR
diff --git a/Engine/Plugins/Runtime/PropertyBindingUtils/Source/PropertyBindingUtils/Private/PropertyBindingPath.cpp b/Engine/Plugins/Runtime/PropertyBindingUtils/Source/PropertyBindingUtils/Private/PropertyBindingPath.cpp
--- a/Engine/Plugins/Runtime/PropertyBindingUtils/Source/PropertyBindingUtils/Private/PropertyBindingPath.cpp
+++ b/Engine/Plugins/Runtime/PropertyBindingUtils/Source/PropertyBindingUtils/Private/PropertyBindingPath.cpp
@@ -1,9 +1,8 @@
 ﻿// Copyright Epic Games, Inc. All Rights Reserved.
 #include "PropertyBindingPath.h"
-#include "UObject/EnumProperty.h"
+#include "PropertyBindingDataView.h"
 #include "Misc/EnumerateRange.h"
 #include "StructUtils/InstancedStructContainer.h"
-#include "StructUtils/PropertyBag.h"
 
 #if WITH_EDITOR
 #include "UObject/CoreRedirects.h"
@@ -11,29 +10,10 @@
 #include "Engine/BlueprintGeneratedClass.h"
 #include "StructUtils/UserDefinedStruct.h"
 #include "Kismet2/StructureEditorUtils.h"
-#include "UObject/Field.h"
 #endif
 
-#include UE_INLINE_GENERATED_CPP_BY_NAME(PropertyBindingPath)
-
-//----------------------------------------------------------------//
-//  FBindableStructDesc
-//----------------------------------------------------------------//
-
-FString FBindableStructDesc::ToString() const
-{
-	FStringBuilderBase Result;
-
-	Result += TEXT(" '");
-	Result += Name.ToString();
-	Result += TEXT("'");
 
-	return Result.ToString();
-}
-
-//----------------------------------------------------------------//
-//  FPropertyBindingPath
-//----------------------------------------------------------------//
+#include UE_INLINE_GENERATED_CPP_BY_NAME(PropertyBindingPath)
 
 bool FPropertyBindingPath::FromString(const FStringView InPath)
 {
@@ -196,7 +176,6 @@ FString FPropertyBindingPath::ToString(const int32 HighlightedSegment, const TCH
 	return Result.ToString();
 }
 
-
 bool FPropertyBindingPath::Includes(const FPropertyBindingPath& Other) const
 {
 #if WITH_EDITORONLY_DATA
@@ -369,7 +348,7 @@ bool FPropertyBindingPath::ResolveIndirectionsWithValue(const FPropertyBindingDa
 			Indirection.ArrayIndex = Segment->GetArrayIndex();
 			Indirection.PropertyOffset = ArrayProperty->GetOffset_ForInternal();
 			Indirection.PathSegmentIndex = Segment.GetIndex();
-			Indirection.AccessType = EPropertyBindingAccessType::IndexArray;
+			Indirection.AccessType = EPropertyBindingPropertyAccessType::IndexArray;
 #if WITH_EDITORONLY_DATA
 			Indirection.RedirectedName = RedirectedName;
 			Indirection.PropertyGuid = PropertyGuid;
@@ -428,7 +407,7 @@ bool FPropertyBindingPath::ResolveIndirectionsWithValue(const FPropertyBindingDa
 		Indirection.ArrayIndex = ArrayIndex;
 		Indirection.PropertyOffset = Offset;
 		Indirection.PathSegmentIndex = Segment.GetIndex();
-		Indirection.AccessType = EPropertyBindingAccessType::Offset; 
+		Indirection.AccessType = EPropertyBindingPropertyAccessType::Offset; 
 #if WITH_EDITORONLY_DATA
 		Indirection.RedirectedName = RedirectedName;
 		Indirection.PropertyGuid = PropertyGuid;
@@ -462,7 +441,7 @@ bool FPropertyBindingPath::ResolveIndirectionsWithValue(const FPropertyBindingDa
 							CurrentAddress = StructMemory;
 							CurrentStruct = ScriptStruct;
 							Indirection.InstanceStruct = CurrentStruct;
-							Indirection.AccessType = EPropertyBindingAccessType::StructInstance;
+							Indirection.AccessType = EPropertyBindingPropertyAccessType::StructInstance;
 							return true;
 						};
 
@@ -497,7 +476,7 @@ bool FPropertyBindingPath::ResolveIndirectionsWithValue(const FPropertyBindingDa
 					{
 						CurrentAddress = (uint8*)CurrentAddress + Offset;
 						CurrentStruct = StructProperty->Struct;
-						Indirection.AccessType = EPropertyBindingAccessType::Offset;
+						Indirection.AccessType = EPropertyBindingPropertyAccessType::Offset;
 					}
 				}
 				else
@@ -506,12 +485,12 @@ bool FPropertyBindingPath::ResolveIndirectionsWithValue(const FPropertyBindingDa
 					{
 						CurrentStruct = Segment->GetInstanceStruct();
 						Indirection.InstanceStruct = CurrentStruct;
-						Indirection.AccessType = EPropertyBindingAccessType::StructInstance;
+						Indirection.AccessType = EPropertyBindingPropertyAccessType::StructInstance;
 					}
 					else
 					{
 						CurrentStruct = StructProperty->Struct;
-						Indirection.AccessType = EPropertyBindingAccessType::Offset;
+						Indirection.AccessType = EPropertyBindingPropertyAccessType::Offset;
 					}
 				}
 			}
@@ -527,12 +506,12 @@ bool FPropertyBindingPath::ResolveIndirectionsWithValue(const FPropertyBindingDa
 					{
 						CurrentStruct = Object->GetClass();
 						Indirection.InstanceStruct = CurrentStruct;
-						Indirection.AccessType = EPropertyBindingAccessType::ObjectInstance;
+						Indirection.AccessType = EPropertyBindingPropertyAccessType::ObjectInstance;
 					}
 					else
 					{
 						CurrentStruct = ObjectProperty->PropertyClass;
-						Indirection.AccessType = EPropertyBindingAccessType::Object;
+						Indirection.AccessType = EPropertyBindingPropertyAccessType::Object;
 					}
 				}
 				else
@@ -541,12 +520,12 @@ bool FPropertyBindingPath::ResolveIndirectionsWithValue(const FPropertyBindingDa
 					{
 						CurrentStruct = Segment->GetInstanceStruct();
 						Indirection.InstanceStruct = CurrentStruct;
-						Indirection.AccessType = EPropertyBindingAccessType::ObjectInstance;
+						Indirection.AccessType = EPropertyBindingPropertyAccessType::ObjectInstance;
 					}
 					else
 					{
 						CurrentStruct = ObjectProperty->PropertyClass;
-						Indirection.AccessType = EPropertyBindingAccessType::Object;
+						Indirection.AccessType = EPropertyBindingPropertyAccessType::Object;
 					}
 				}
 			}
@@ -570,7 +549,7 @@ bool FPropertyBindingPath::ResolveIndirectionsWithValue(const FPropertyBindingDa
 					CurrentStruct = WeakObjectProperty->PropertyClass;
 				}
 
-				Indirection.AccessType = EPropertyBindingAccessType::WeakObject;
+				Indirection.AccessType = EPropertyBindingPropertyAccessType::WeakObject;
 			}
 			// Check to see if this is a simple soft object property (eg. not an array of soft objects).
 			else if (const FSoftObjectProperty* SoftObjectProperty = CastField<FSoftObjectProperty>(Property))
@@ -592,7 +571,7 @@ bool FPropertyBindingPath::ResolveIndirectionsWithValue(const FPropertyBindingDa
 					CurrentStruct = SoftObjectProperty->PropertyClass;
 				}
 
-				Indirection.AccessType = EPropertyBindingAccessType::SoftObject;
+				Indirection.AccessType = EPropertyBindingPropertyAccessType::SoftObject;
 			}
 			else
 			{
diff --git a/Engine/Plugins/Runtime/PropertyBindingUtils/Source/PropertyBindingUtils/Public/PropertyBindingBindingCollectionOwner.h b/Engine/Plugins/Runtime/PropertyBindingUtils/Source/PropertyBindingUtils/Public/PropertyBindingBindingCollectionOwner.h
new file mode 100644
--- /dev/null
+++ b/Engine/Plugins/Runtime/PropertyBindingUtils/Source/PropertyBindingUtils/Public/PropertyBindingBindingCollectionOwner.h
@@ -0,0 +1,100 @@
+// Copyright Epic Games, Inc. All Rights Reserved.
+
+#pragma once
+
+#include "UObject/Interface.h"
+#include "PropertyBindingBindingCollectionOwner.generated.h"
+
+namespace UE::PropertyBinding
+{
+struct FPropertyCreationDescriptor;
+}
+struct FPropertyBindingPath;
+struct FPropertyBindingBindingCollection;
+struct FPropertyBindingBindableStructDescriptor;
+struct FPropertyBindingDataView;
+template<typename BaseStructT> struct TInstancedStruct;
+
+
+UINTERFACE(MinimalAPI)
+class UPropertyBindingBindingCollectionOwner : public UInterface
+{
+	GENERATED_BODY()
+};
+
+/**
+ * Interface that asset or their associated Editor-only data could inherit from to facilitate
+ * operations on property bindings.
+ * @see FPropertyBindingExtension
+ */
+class IPropertyBindingBindingCollectionOwner
+{
+	GENERATED_BODY()
+
+#if WITH_EDITOR
+public:
+	
+	/**
+	 * Returns structs within the owner that are visible for target struct.
+	 * @param InTargetStructID Target struct ID
+	 * @param OutStructDescs Result descriptors of the visible structs.
+	 */
+	PROPERTYBINDINGUTILS_API virtual void GetBindableStructs(const FGuid InTargetStructID, TArray<TInstancedStruct<FPropertyBindingBindableStructDescriptor>>& OutStructDescs) const = 0;
+
+	/**
+	 * Returns struct descriptor based on struct ID.
+	 * @param InStructID Target struct ID
+	 * @param OutStructDesc Result descriptor.
+	 * @return True if struct found.
+	 */
+	PROPERTYBINDINGUTILS_API virtual bool GetBindableStructByID(const FGuid InStructID, TInstancedStruct<FPropertyBindingBindableStructDescriptor>& OutStructDesc) const = 0;
+		
+	/**
+	 * Returns data view based on struct ID.
+	 * @param InStructID Target struct ID
+	 * @param OutDataView Result data view.
+	 * @return True if struct found.
+	 */
+	PROPERTYBINDINGUTILS_API virtual bool GetBindingDataViewByID(const FGuid InStructID, FPropertyBindingDataView& OutDataView) const = 0;
+
+	/** @return Pointer to editor property bindings. */
+	PROPERTYBINDINGUTILS_API virtual FPropertyBindingBindingCollection* GetEditorPropertyBindings() = 0;
+
+	/** @return Pointer to const editor property bindings. */
+	PROPERTYBINDINGUTILS_API virtual const FPropertyBindingBindingCollection* GetEditorPropertyBindings() const = 0;
+	
+	/**
+	 * Can be overridden to determine whether the struct matching the given struct id is capable of adding new properties.
+	 * @param InStructID Target struct ID
+	 * @return Whether the struct supports adding new properties
+	 */
+	PROPERTYBINDINGUTILS_API virtual bool CanCreateParameter(const FGuid InStructID) const
+	{
+		return true;
+	}
+
+	/**
+	 * Creates the given properties in the property bag of the struct matching the given struct ID
+	 * @param InStructID Target struct ID
+	 * @param InOutCreationDescs the descriptions of the properties to create. This is modified to update the property names that actually got created
+	 */
+	PROPERTYBINDINGUTILS_API virtual void CreateParametersForStruct(const FGuid InStructID, TArrayView<UE::PropertyBinding::FPropertyCreationDescriptor> InOutCreationDescs) = 0;
+
+	/** Can be overridden to provide a fallback structure ID to build the property path if no bindable structs are found when traversing a PropertyHandle hierarchy. */
+	PROPERTYBINDINGUTILS_API virtual FGuid GetFallbackStructID() const
+	{
+		// No fallback by default
+		return {};
+	}
+
+	/** Can be overridden to provide additional bindable structs from property functions. */
+	PROPERTYBINDINGUTILS_API virtual void AppendBindablePropertyFunctionStructs(TArray<TInstancedStruct<FPropertyBindingBindableStructDescriptor>>& InOutStructs) const
+	{
+	}
+
+	/** Can be overridden to perform additional operations when property bindings changed. */
+	PROPERTYBINDINGUTILS_API virtual void OnPropertyBindingChanged(const FPropertyBindingPath& InSourcePath, const FPropertyBindingPath& InTargetPath)
+	{
+	}
+#endif // WITH_EDITOR
+};
diff --git a/Engine/Plugins/Runtime/PropertyBindingUtils/Source/PropertyBindingUtils/Public/PropertyBindingPath.h b/Engine/Plugins/Runtime/PropertyBindingUtils/Source/PropertyBindingUtils/Public/PropertyBindingPath.h
--- a/Engine/Plugins/Runtime/PropertyBindingUtils/Source/PropertyBindingUtils/Public/PropertyBindingPath.h
+++ b/Engine/Plugins/Runtime/PropertyBindingUtils/Source/PropertyBindingUtils/Public/PropertyBindingPath.h
@@ -2,259 +2,18 @@
 
 #pragma once
 
-#include "EngineDefines.h"
-#include "Engine/EngineTypes.h"
-#include "StructUtils/StructView.h"
-#include "Misc/Guid.h"
-#include "PropertyBindingPath.generated.h"
-
-struct FPropertyBindingPath;
-
-
-
-/**
- * Descriptor for a struct or class that can be a binding source or target.
- * Each struct has unique identifier, which is used to distinguish them, and name that is mostly for debugging and UI.
- */
-USTRUCT()
-struct PROPERTYBINDINGUTILS_API FBindableStructDesc
-{
-	GENERATED_BODY()
-
-	FBindableStructDesc() = default;
-
-#if WITH_EDITORONLY_DATA
-	FBindableStructDesc(const FName InName, const UStruct* InStruct, const FGuid InGuid)
-		: Struct(InStruct)
-		, Name(InName)
-		, ID(InGuid)
-	{
-	}
-
-	bool operator==(const FBindableStructDesc& RHS) const
-	{
-		return ID == RHS.ID && Struct == RHS.Struct; // Not checking name, it's cosmetic.
-	}
-#endif
-
-	bool IsValid() const { return Struct != nullptr; }
-
-	FString ToString() const;
-	
-	/** The type of the struct or class. */
-	UPROPERTY()
-	TObjectPtr<const UStruct> Struct = nullptr;
-
-	/** Name of the struct (used for debugging, logging, cosmetic). */
-	UPROPERTY()
-	FName Name;
-
-#if WITH_EDITORONLY_DATA
-	/** Unique identifier of the struct. */
-	UPROPERTY()
-	FGuid ID;
-#endif
-};
-
-
-UENUM()
-enum class EPropertyBindingAccessType : uint8
-{
-	Offset,			// Access node is a simple basePtr + offset
-	Object,			// Access node needs to dereference an object at its current address
-	WeakObject,		// Access is a weak object
-	SoftObject,		// Access is a soft object
-	ObjectInstance,	// Access node needs to dereference an object of specific type at its current address
-	StructInstance,	// Access node needs to dereference an instanced struct of specific type at its current address
-	IndexArray,		// Access node indexes a dynamic array
-};
-
-/**
- * Short-lived pointer to an object or struct.
- * Can be constructed directly from a pointer to a UObject or from a StructView.
- *
- * When creating from a generic UStruct, the data view expects a valid type (UStruct) when a valid memory pointer is provided.
- * It that case it is also fine to provide null for both and the constructed data view will be considered as invalid.
- */
-struct FPropertyBindingDataView
-{
-	FPropertyBindingDataView() = default;
-
-	/**
-	 * Generic struct constructor.
-	 * Valid UStruct is required when passing in valid memory pointer.
-	 * Both can be null and the constructed data view will be considered as invalid.
-	 * @see IsValid
-	 */
-	FPropertyBindingDataView(const UStruct* InStruct, void* InMemory) : Struct(InStruct), Memory(InMemory)
-	{
-		// Must have type with valid pointer.
-		check(!Memory || (Memory && Struct));
-	}
-
-	/** UObject constructor. */
-	FPropertyBindingDataView(UObject* Object) : Struct(Object ? Object->GetClass() : nullptr), Memory(Object)
-	{
-		// Must have type with valid pointer.
-		check(!Memory || (Memory && Struct));
-	}
-
-	/** Struct from a StructView. */
-	FPropertyBindingDataView(const FStructView StructView) : Struct(StructView.GetScriptStruct()), Memory(StructView.GetMemory())
-	{
-		// Must have type with valid pointer.
-		check(!Memory || (Memory && Struct));
-	}
-
-	/**
-	 * Check is the view is valid (both pointer and type are set). On valid views it is safe to call the Get<>() methods returning a reference.
-	 * @return True if the view is valid.
-	 */
-	bool IsValid() const
-	{
-		return Memory != nullptr && Struct != nullptr;
-	}
-
-	/**
-	 * UObject getters (reference & pointer, const & mutable)
-	 */
-	template <typename T>
-    typename TEnableIf<TIsDerivedFrom<T, UObject>::IsDerived, const T&>::Type Get() const
-	{
-		check(Memory != nullptr);
-		check(Struct != nullptr);
-		check(Struct->IsChildOf(T::StaticClass()));
-		return *((T*)Memory);
-	}
-
-	template <typename T>
-	typename TEnableIf<TIsDerivedFrom<T, UObject>::IsDerived, T&>::Type GetMutable() const
-	{
-		check(Memory != nullptr);
-		check(Struct != nullptr);
-		check(Struct->IsChildOf(T::StaticClass()));
-		return *((T*)Memory);
-	}
-
-	template <typename T>
-	typename TEnableIf<TIsDerivedFrom<T, UObject>::IsDerived, const T*>::Type GetPtr() const
-	{
-		// If Memory is set, expect Struct too. Otherwise, let nulls pass through.
-		check(!Memory || (Memory && Struct));
-		check(!Struct || Struct->IsChildOf(T::StaticClass()));
-		return ((T*)Memory);
-	}
-
-	template <typename T>
-    typename TEnableIf<TIsDerivedFrom<T, UObject>::IsDerived, T*>::Type GetMutablePtr() const
-	{
-		// If Memory is set, expect Struct too. Otherwise, let nulls pass through.
-		check(!Memory || (Memory && Struct));
-		check(!Struct || Struct->IsChildOf(T::StaticClass()));
-		return ((T*)Memory);
-	}
-
-	/**
-	 * Struct getters (reference & pointer, const & mutable)
-	 */
-	template <typename T>
-	typename TEnableIf<!TIsDerivedFrom<T, UObject>::IsDerived && !TIsIInterface<T>::Value, const T&>::Type Get() const
-	{
-		check(Memory != nullptr);
-		check(Struct != nullptr);
-		check(Struct->IsChildOf(T::StaticStruct()));
-		return *((T*)Memory);
-	}
-
-	template <typename T>
-    typename TEnableIf<!TIsDerivedFrom<T, UObject>::IsDerived && !TIsIInterface<T>::Value, T&>::Type GetMutable() const
-	{
-		check(Memory != nullptr);
-		check(Struct != nullptr);
-		check(Struct->IsChildOf(T::StaticStruct()));
-		return *((T*)Memory);
-	}
-
-	template <typename T>
-    typename TEnableIf<!TIsDerivedFrom<T, UObject>::IsDerived && !TIsIInterface<T>::Value, const T*>::Type GetPtr() const
-	{
-		// If Memory is set, expect Struct too. Otherwise, let nulls pass through.
-		check(!Memory || (Memory && Struct));
-		check(!Struct || Struct->IsChildOf(T::StaticStruct()));
-		return ((T*)Memory);
-	}
-
-	template <typename T>
-    typename TEnableIf<!TIsDerivedFrom<T, UObject>::IsDerived && !TIsIInterface<T>::Value, T*>::Type GetMutablePtr() const
-	{
-		// If Memory is set, expect Struct too. Otherwise, let nulls pass through.
-		check(!Memory || (Memory && Struct));
-		check(!Struct || Struct->IsChildOf(T::StaticStruct()));
-		return ((T*)Memory);
-	}
-
-	/**
-	 * IInterface getters (reference & pointer, const & mutable)
-	 */
-	template <typename T>
-	typename TEnableIf<TIsIInterface<T>::Value, const T&>::Type Get() const
-	{
-		check(!Memory || (Memory && Struct));
-		check(Struct->IsChildOf(UObject::StaticClass()) && ((UClass*)Struct)->ImplementsInterface(T::UClassType::StaticClass()));
-		return *(T*)((UObject*)Memory)->GetInterfaceAddress(T::UClassType::StaticClass());
-	}
-
-	template <typename T>
-    typename TEnableIf<TIsIInterface<T>::Value, T&>::Type GetMutable() const
-	{
-		check(!Memory || (Memory && Struct));
-		check(Struct->IsChildOf(UObject::StaticClass()) && ((UClass*)Struct)->ImplementsInterface(T::UClassType::StaticClass()));
-		return *(T*)((UObject*)Memory)->GetInterfaceAddress(T::UClassType::StaticClass());
-	}
-
-	template <typename T>
-    typename TEnableIf<TIsIInterface<T>::Value, const T*>::Type GetPtr() const
-	{
-		check(!Memory || (Memory && Struct));
-		check(Struct->IsChildOf(UObject::StaticClass()) && ((UClass*)Struct)->ImplementsInterface(T::UClassType::StaticClass()));
-		return (T*)((UObject*)Memory)->GetInterfaceAddress(T::UClassType::StaticClass());
-	}
-
-	template <typename T>
-    typename TEnableIf<TIsIInterface<T>::Value, T*>::Type GetMutablePtr() const
-	{
-		check(!Memory || (Memory && Struct));
-		check(Struct->IsChildOf(UObject::StaticClass()) && ((UClass*)Struct)->ImplementsInterface(T::UClassType::StaticClass()));
-		return (T*)((UObject*)Memory)->GetInterfaceAddress(T::UClassType::StaticClass());
-	}
-
-	/** @return Struct describing the data type. */
-	const UStruct* GetStruct() const
-	{
-		return Struct;
-	}
-
-	/** @return Raw const pointer to the data. */
-	const void* GetMemory() const
-	{
-		return Memory;
-	}
+#include "PropertyBindingTypes.h"
 
-	/** @return Raw mutable pointer to the data. */
-	void* GetMutableMemory() const
-	{
-		return Memory;
-	}
-	
-protected:
-	/** UClass or UScriptStruct of the data. */
-	const UStruct* Struct = nullptr;
+#if UE_ENABLE_INCLUDE_ORDER_DEPRECATED_IN_5_6
+#include "PropertyBindingBindableStructDescriptor.h"
+#include "PropertyBindingDataView.h"
 
-	/** Memory pointing at the class or struct */
-	void* Memory = nullptr;
-};
+#endif // UE_ENABLE_INCLUDE_ORDER_DEPRECATED_IN_5_6
 
+#include "PropertyBindingPath.generated.h"
 
+struct FPropertyBindingDataView;
+struct FPropertyBindingPath;
 
 /**
  * Struct describing an indirection at specific segment at path.
@@ -304,7 +63,7 @@ struct FPropertyBindingPathIndirection
 		return PathSegmentIndex;
 	}
 
-	EPropertyBindingAccessType GetAccessType() const
+	EPropertyBindingPropertyAccessType GetAccessType() const
 	{
 		return AccessType;
 	}
@@ -354,7 +113,7 @@ private:
 	int32 PathSegmentIndex = 0;
 	
 	/** How to access the data through the indirection. */
-	EPropertyBindingAccessType AccessType = EPropertyBindingAccessType::Offset;
+	EPropertyBindingPropertyAccessType AccessType = EPropertyBindingPropertyAccessType::Offset;
 
 #if WITH_EDITORONLY_DATA
 	/** Redirected name, if the give property name was not found but was reconciled using core redirect or property Guid. Requires ResolveIndirections/WithValue() to be called with bHandleRedirects = true. */
@@ -632,4 +391,4 @@ private:
 	/** Path segments pointing to a specific property on the path. */
 	UPROPERTY()
 	TArray<FPropertyBindingPathSegment> Segments;
-};
\ No newline at end of file
+};
diff --git a/Engine/Plugins/Runtime/PropertyBindingUtils/Source/PropertyBindingUtilsTestSuite/Private/PropertyBindingUtilsTest.cpp b/Engine/Plugins/Runtime/PropertyBindingUtils/Source/PropertyBindingUtilsTestSuite/Private/PropertyBindingUtilsTest.cpp
--- a/Engine/Plugins/Runtime/PropertyBindingUtils/Source/PropertyBindingUtilsTestSuite/Private/PropertyBindingUtilsTest.cpp
+++ b/Engine/Plugins/Runtime/PropertyBindingUtils/Source/PropertyBindingUtilsTestSuite/Private/PropertyBindingUtilsTest.cpp
@@ -3,6 +3,7 @@
 #include "PropertyBindingUtilsTest.h"
 #include "PropertyBindingPath.h"
 #include "AITestsCommon.h"
+#include "PropertyBindingDataView.h"
 #include "Engine/World.h"
 
 #include UE_INLINE_GENERATED_CPP_BY_NAME(PropertyBindingUtilsTest)
@@ -30,8 +31,8 @@ struct FPropertyBindingUtilsTest_PropertyPathOffset : FAITestBase
 		AITEST_EQUAL("Should have no resolve errors", ResolveErrors.Len(), 0);
 		
 		AITEST_EQUAL("Should have 2 indirections", Indirections.Num(), 2);
-		AITEST_EQUAL("Indirection 0 should be Offset type", Indirections[0].GetAccessType(), EPropertyBindingAccessType::Offset);
-		AITEST_EQUAL("Indirection 1 should be Offset type", Indirections[1].GetAccessType(), EPropertyBindingAccessType::Offset);
+		AITEST_EQUAL("Indirection 0 should be Offset type", Indirections[0].GetAccessType(), EPropertyBindingPropertyAccessType::Offset);
+		AITEST_EQUAL("Indirection 1 should be Offset type", Indirections[1].GetAccessType(), EPropertyBindingPropertyAccessType::Offset);
 
 		return true;
 	}
@@ -145,7 +146,7 @@ struct FPropertyBindingUtilsTest_PropertyPathWrongObject : FAITestBase
 
 			AITEST_TRUE("Resolve path should succeeed", bResolveResult);
 			AITEST_EQUAL("Should have 2 indirections", Indirections.Num(), 2);
-			AITEST_TRUE("Object ", Indirections[0].GetAccessType() == EPropertyBindingAccessType::ObjectInstance);
+			AITEST_TRUE("Object ", Indirections[0].GetAccessType() == EPropertyBindingPropertyAccessType::ObjectInstance);
 			AITEST_TRUE("Object ", Indirections[0].GetContainerStruct() == Object->GetClass());
 			AITEST_TRUE("Object ", Indirections[0].GetInstanceStruct() == UPropertyBindingUtilsTest_PropertyObjectInstancedWithB::StaticClass());
 			AITEST_EQUAL("Should not have error", ResolveErrors.Len(), 0);
@@ -188,8 +189,8 @@ struct FPropertyBindingUtilsTest_PropertyPathArray : FAITestBase
 		AITEST_TRUE("Resolve path should succeeed", bResolveResult);
 		AITEST_EQUAL("Should have no resolve errors", ResolveErrors.Len(), 0);
 		AITEST_EQUAL("Should have 2 indirections", Indirections.Num(), 2);
-		AITEST_EQUAL("Indirection 0 should be IndexArray type", Indirections[0].GetAccessType(), EPropertyBindingAccessType::IndexArray);
-		AITEST_EQUAL("Indirection 1 should be Offset type", Indirections[1].GetAccessType(), EPropertyBindingAccessType::Offset);
+		AITEST_EQUAL("Indirection 0 should be IndexArray type", Indirections[0].GetAccessType(), EPropertyBindingPropertyAccessType::IndexArray);
+		AITEST_EQUAL("Indirection 1 should be Offset type", Indirections[1].GetAccessType(), EPropertyBindingPropertyAccessType::Offset);
 
 		const int32 Value = *reinterpret_cast<const int32*>(Indirections[1].GetPropertyAddress());
 		AITEST_EQUAL("Value should be 123", Value, 123);
@@ -247,9 +248,9 @@ struct FPropertyBindingUtilsTest_PropertyPathArrayOfStructs : FAITestBase
 			AITEST_TRUE("Resolve path1 should succeeed", bResolveResult);
 			AITEST_EQUAL("Should have no resolve errors", ResolveErrors.Len(), 0);
 			AITEST_EQUAL("Should have 3 indirections", Indirections.Num(), 3);
-			AITEST_EQUAL("Indirection 0 should be ArrayIndex type", Indirections[0].GetAccessType(), EPropertyBindingAccessType::IndexArray);
-			AITEST_EQUAL("Indirection 1 should be Offset type", Indirections[1].GetAccessType(), EPropertyBindingAccessType::Offset);
-			AITEST_EQUAL("Indirection 2 should be Offset type", Indirections[2].GetAccessType(), EPropertyBindingAccessType::Offset);
+			AITEST_EQUAL("Indirection 0 should be ArrayIndex type", Indirections[0].GetAccessType(), EPropertyBindingPropertyAccessType::IndexArray);
+			AITEST_EQUAL("Indirection 1 should be Offset type", Indirections[1].GetAccessType(), EPropertyBindingPropertyAccessType::Offset);
+			AITEST_EQUAL("Indirection 2 should be Offset type", Indirections[2].GetAccessType(), EPropertyBindingPropertyAccessType::Offset);
 
 			const int32 Value = *reinterpret_cast<const int32*>(Indirections[2].GetPropertyAddress());
 			AITEST_EQUAL("Value should be 3", Value, 3);
@@ -263,10 +264,10 @@ struct FPropertyBindingUtilsTest_PropertyPathArrayOfStructs : FAITestBase
 			AITEST_TRUE("Resolve path2 should succeeed", bResolveResult);
 			AITEST_EQUAL("Should have no resolve errors", ResolveErrors.Len(), 0);
 			AITEST_EQUAL("Should have 4 indirections", Indirections.Num(), 4);
-			AITEST_EQUAL("Indirection 0 should be ArrayIndex type", Indirections[0].GetAccessType(), EPropertyBindingAccessType::IndexArray);
-			AITEST_EQUAL("Indirection 1 should be Offset type", Indirections[1].GetAccessType(), EPropertyBindingAccessType::Offset);
-			AITEST_EQUAL("Indirection 2 should be Offset type", Indirections[2].GetAccessType(), EPropertyBindingAccessType::Offset);
-			AITEST_EQUAL("Indirection 3 should be Offset type", Indirections[3].GetAccessType(), EPropertyBindingAccessType::Offset);
+			AITEST_EQUAL("Indirection 0 should be ArrayIndex type", Indirections[0].GetAccessType(), EPropertyBindingPropertyAccessType::IndexArray);
+			AITEST_EQUAL("Indirection 1 should be Offset type", Indirections[1].GetAccessType(), EPropertyBindingPropertyAccessType::Offset);
+			AITEST_EQUAL("Indirection 2 should be Offset type", Indirections[2].GetAccessType(), EPropertyBindingPropertyAccessType::Offset);
+			AITEST_EQUAL("Indirection 3 should be Offset type", Indirections[3].GetAccessType(), EPropertyBindingPropertyAccessType::Offset);
 
 			const int32 Value = *reinterpret_cast<const int32*>(Indirections[3].GetPropertyAddress());
 			AITEST_EQUAL("Value should be 42", Value, 42);
@@ -303,9 +304,9 @@ struct FPropertyBindingUtilsTest_PropertyPathArrayOfInstancedObjects : FAITestBa
 			AITEST_TRUE("Resolve path should succeeed", bResolveResult);
 			AITEST_EQUAL("Should have no resolve errors", ResolveErrors.Len(), 0);
 			AITEST_EQUAL("Should have 3 indirections", Indirections.Num(), 3);
-			AITEST_EQUAL("Indirection 0 should be ArrayIndex type", Indirections[0].GetAccessType(), EPropertyBindingAccessType::IndexArray);
-			AITEST_EQUAL("Indirection 1 should be StructInstance type", Indirections[1].GetAccessType(), EPropertyBindingAccessType::StructInstance);
-			AITEST_EQUAL("Indirection 2 should be Offset type", Indirections[2].GetAccessType(), EPropertyBindingAccessType::Offset);
+			AITEST_EQUAL("Indirection 0 should be ArrayIndex type", Indirections[0].GetAccessType(), EPropertyBindingPropertyAccessType::IndexArray);
+			AITEST_EQUAL("Indirection 1 should be StructInstance type", Indirections[1].GetAccessType(), EPropertyBindingPropertyAccessType::StructInstance);
+			AITEST_EQUAL("Indirection 2 should be Offset type", Indirections[2].GetAccessType(), EPropertyBindingPropertyAccessType::Offset);
 		}
 
 		{
@@ -316,9 +317,9 @@ struct FPropertyBindingUtilsTest_PropertyPathArrayOfInstancedObjects : FAITestBa
 			AITEST_TRUE("Resolve path should succeeed", bResolveResult);
 			AITEST_EQUAL("Should have no resolve errors", ResolveErrors.Len(), 0);
 			AITEST_EQUAL("Should have 3 indirections", Indirections.Num(), 3);
-			AITEST_EQUAL("Indirection 0 should be ArrayIndex type", Indirections[0].GetAccessType(), EPropertyBindingAccessType::IndexArray);
-			AITEST_EQUAL("Indirection 1 should be StructInstance type", Indirections[1].GetAccessType(), EPropertyBindingAccessType::StructInstance);
-			AITEST_EQUAL("Indirection 2 should be Offset type", Indirections[2].GetAccessType(), EPropertyBindingAccessType::Offset);
+			AITEST_EQUAL("Indirection 0 should be ArrayIndex type", Indirections[0].GetAccessType(), EPropertyBindingPropertyAccessType::IndexArray);
+			AITEST_EQUAL("Indirection 1 should be StructInstance type", Indirections[1].GetAccessType(), EPropertyBindingPropertyAccessType::StructInstance);
+			AITEST_EQUAL("Indirection 2 should be Offset type", Indirections[2].GetAccessType(), EPropertyBindingPropertyAccessType::Offset);
 
 			const int32 Value = *reinterpret_cast<const int32*>(Indirections[2].GetPropertyAddress());
 			AITEST_EQUAL("Value should be 123", Value, 123);
diff --git a/Engine/Plugins/Runtime/StateTree/Source/StateTreeEditorModule/Private/AsyncStateTreeDiff.cpp b/Engine/Plugins/Runtime/StateTree/Source/StateTreeEditorModule/Private/AsyncStateTreeDiff.cpp
--- a/Engine/Plugins/Runtime/StateTree/Source/StateTreeEditorModule/Private/AsyncStateTreeDiff.cpp
+++ b/Engine/Plugins/Runtime/StateTree/Source/StateTreeEditorModule/Private/AsyncStateTreeDiff.cpp
@@ -192,12 +192,12 @@ static bool AreStateTreePropertiesEqual(const UStateTreeEditorData* StateTreeDat
 	bool bBindingsEqual = StateTreeDataA->EditorBindings.GetBindings().Num() == StateTreeDataB->EditorBindings.GetBindings().Num();
 	if (bBindingsEqual)
 	{
-		for (const FStateTreePropertyPathBinding& PropertyPathBinding : StateTreeDataA->EditorBindings.GetBindings())
+		for (const FPropertyBindingBinding& PropertyPathBinding : StateTreeDataA->EditorBindings.GetBindings())
 		{
 			const FStateTreePropertyPath& PropertyPathTarget = PropertyPathBinding.GetTargetPath();
-			if (StateTreeDataB->EditorBindings.HasPropertyBinding(PropertyPathBinding.GetTargetPath()))
+			if (StateTreeDataB->EditorBindings.HasBinding(PropertyPathBinding.GetTargetPath()))
 			{
-				if (*StateTreeDataA->EditorBindings.GetPropertyBindingSource(PropertyPathTarget) != *StateTreeDataB->EditorBindings.GetPropertyBindingSource(PropertyPathTarget))
+				if (*StateTreeDataA->EditorBindings.GetBindingSource(PropertyPathTarget) != *StateTreeDataB->EditorBindings.GetBindingSource(PropertyPathTarget))
 				{
 					bBindingsEqual = false;
 					break;
@@ -294,7 +294,7 @@ static void GetBindingsDifferences(UStateTreeEditorData* StateTreeDataA, UStateT
 
 	TArray<FBindingDiff> BindingDiffs;
 	// Check the differences in Bindings
-	for (const FStateTreePropertyPathBinding& PropertyPathBinding : StateTreeDataA->EditorBindings.GetBindings())
+	for (const FPropertyBindingBinding& PropertyPathBinding : StateTreeDataA->EditorBindings.GetBindings())
 	{
 		FStateTreePropertyPath PropertyPathTarget = PropertyPathBinding.GetTargetPath();
 		FStateTreePropertyPath PropertyPathSource = PropertyPathBinding.GetSourcePath();
@@ -304,7 +304,7 @@ static void GetBindingsDifferences(UStateTreeEditorData* StateTreeDataA, UStateT
 		Entry.SourcePathA = PropertyPathSource;
 		BindingDiffs.Add(Entry);
 	}
-	for (const FStateTreePropertyPathBinding& PropertyPathBinding : StateTreeDataB->EditorBindings.GetBindings())
+	for (const FPropertyBindingBinding& PropertyPathBinding : StateTreeDataB->EditorBindings.GetBindings())
 	{
 		FStateTreePropertyPath PropertyPathTarget = PropertyPathBinding.GetTargetPath();
 		FStateTreePropertyPath PropertyPathSource = PropertyPathBinding.GetSourcePath();
