Fix CIS Issue 715116: export destructor

diff --git a/Engine/Plugins/Runtime/PropertyBindingUtils/Source/PropertyBindingUtils/Public/PropertyBindingBindableStructDescriptor.h b/Engine/Plugins/Runtime/PropertyBindingUtils/Source/PropertyBindingUtils/Public/PropertyBindingBindableStructDescriptor.h
--- a/Engine/Plugins/Runtime/PropertyBindingUtils/Source/PropertyBindingUtils/Public/PropertyBindingBindableStructDescriptor.h
+++ b/Engine/Plugins/Runtime/PropertyBindingUtils/Source/PropertyBindingUtils/Public/PropertyBindingBindableStructDescriptor.h
@@ -14,7 +14,7 @@ struct FPropertyBindingBindableStructDescriptor
 	GENERATED_BODY()
 
 	FPropertyBindingBindableStructDescriptor() = default;
-	virtual ~FPropertyBindingBindableStructDescriptor() = default;
+	PROPERTYBINDINGUTILS_API virtual ~FPropertyBindingBindableStructDescriptor() = default;
 
 #if WITH_EDITORONLY_DATA
 	FPropertyBindingBindableStructDescriptor(const FName InName, const UStruct* InStruct, const FGuid InGuid)
