Fix CIS Issue 715116: define destructor

diff --git a/Engine/Plugins/Runtime/PropertyBindingUtils/Source/PropertyBindingUtils/Public/PropertyBindingBindableStructDescriptor.h b/Engine/Plugins/Runtime/PropertyBindingUtils/Source/PropertyBindingUtils/Public/PropertyBindingBindableStructDescriptor.h
--- a/Engine/Plugins/Runtime/PropertyBindingUtils/Source/PropertyBindingUtils/Public/PropertyBindingBindableStructDescriptor.h
+++ b/Engine/Plugins/Runtime/PropertyBindingUtils/Source/PropertyBindingUtils/Public/PropertyBindingBindableStructDescriptor.h
@@ -14,7 +14,9 @@ struct FPropertyBindingBindableStructDescriptor
 	GENERATED_BODY()
 
 	FPropertyBindingBindableStructDescriptor() = default;
-	PROPERTYBINDINGUTILS_API virtual ~FPropertyBindingBindableStructDescriptor() = default;
+	PROPERTYBINDINGUTILS_API virtual ~FPropertyBindingBindableStructDescriptor()
+	{
+	}
 
 #if WITH_EDITORONLY_DATA
 	FPropertyBindingBindableStructDescriptor(const FName InName, const UStruct* InStruct, const FGuid InGuid)
