Fix CIS Issue 715116: moved destructor to source

diff --git a/Engine/Plugins/Runtime/PropertyBindingUtils/Source/PropertyBindingUtils/Private/PropertyBindingBindableStructDescriptor.cpp b/Engine/Plugins/Runtime/PropertyBindingUtils/Source/PropertyBindingUtils/Private/PropertyBindingBindableStructDescriptor.cpp
--- a/Engine/Plugins/Runtime/PropertyBindingUtils/Source/PropertyBindingUtils/Private/PropertyBindingBindableStructDescriptor.cpp
+++ b/Engine/Plugins/Runtime/PropertyBindingUtils/Source/PropertyBindingUtils/Private/PropertyBindingBindableStructDescriptor.cpp
@@ -2,6 +2,9 @@
 
 #include "PropertyBindingBindableStructDescriptor.h"
 
+FPropertyBindingBindableStructDescriptor::~FPropertyBindingBindableStructDescriptor()
+{
+}
 
 FString FPropertyBindingBindableStructDescriptor::ToString() const
 {
@@ -12,4 +15,4 @@ FString FPropertyBindingBindableStructDescriptor::ToString() const
 	Result += TEXT("'");
 
 	return Result.ToString();
-}
\ No newline at end of file
+}
diff --git a/Engine/Plugins/Runtime/PropertyBindingUtils/Source/PropertyBindingUtils/Public/PropertyBindingBindableStructDescriptor.h b/Engine/Plugins/Runtime/PropertyBindingUtils/Source/PropertyBindingUtils/Public/PropertyBindingBindableStructDescriptor.h
--- a/Engine/Plugins/Runtime/PropertyBindingUtils/Source/PropertyBindingUtils/Public/PropertyBindingBindableStructDescriptor.h
+++ b/Engine/Plugins/Runtime/PropertyBindingUtils/Source/PropertyBindingUtils/Public/PropertyBindingBindableStructDescriptor.h
@@ -14,9 +14,7 @@ struct FPropertyBindingBindableStructDescriptor
 	GENERATED_BODY()
 
 	FPropertyBindingBindableStructDescriptor() = default;
-	PROPERTYBINDINGUTILS_API virtual ~FPropertyBindingBindableStructDescriptor()
-	{
-	}
+	PROPERTYBINDINGUTILS_API virtual ~FPropertyBindingBindableStructDescriptor();
 
 #if WITH_EDITORONLY_DATA
 	FPropertyBindingBindableStructDescriptor(const FName InName, const UStruct* InStruct, const FGuid InGuid)
