Removed undesired todo comment

diff --git a/Engine/Plugins/Runtime/PropertyBindingUtils/Source/PropertyBindingUtilsEditor/Private/PropertyBindingExtension.cpp b/Engine/Plugins/Runtime/PropertyBindingUtils/Source/PropertyBindingUtilsEditor/Private/PropertyBindingExtension.cpp
--- a/Engine/Plugins/Runtime/PropertyBindingUtils/Source/PropertyBindingUtilsEditor/Private/PropertyBindingExtension.cpp
+++ b/Engine/Plugins/Runtime/PropertyBindingUtils/Source/PropertyBindingUtilsEditor/Private/PropertyBindingExtension.cpp
@@ -1117,7 +1117,6 @@ bool FPropertyBindingExtension::IsPropertyExtendable(const UClass* InObjectClass
 		{
 			FPropertyBindingPath TargetPath;
 
-			// @tod YO : debug SO + ST
 			// Figure out the structs we're editing, and property path relative to current property.
 			const FGuid FallbackStructID = BindingsOwner->GetFallbackStructID();
 			const TSharedPtr<const IPropertyHandle> SharedHandle = UE::PropertyBinding::MakeStructPropertyPathFromPropertyHandle(
