StateTree: Add the node indexes to the description. It helps when reading big assets. Enable the log with cvar StateTree.Compiler.LogResultOnCompilationCompleted.
#rnx

diff --git a/Engine/Plugins/Runtime/PropertyBindingUtils/Source/PropertyBindingUtils/Private/PropertyBindingBindingCollection.cpp b/Engine/Plugins/Runtime/PropertyBindingUtils/Source/PropertyBindingUtils/Private/PropertyBindingBindingCollection.cpp
--- a/Engine/Plugins/Runtime/PropertyBindingUtils/Source/PropertyBindingUtils/Private/PropertyBindingBindingCollection.cpp
+++ b/Engine/Plugins/Runtime/PropertyBindingUtils/Source/PropertyBindingUtils/Private/PropertyBindingBindingCollection.cpp
@@ -1095,13 +1095,16 @@ FString FPropertyBindingBindingCollection::DebugAsString() const
 	DebugString.Appendf(TEXT("\nSourceStructs (%d)\n"), GetNumBindableStructDescriptors());
 	if (GetNumBindableStructDescriptors())
 	{
-		DebugString.Appendf(TEXT("  [ %-45s | % -40s | % -55s ]\n"), TEXT("Type"), TEXT("DataSource"), TEXT("Name"));
-		VisitSourceStructDescriptorInternal([&DebugString](const FPropertyBindingBindableStructDescriptor& Descriptor)
+		DebugString.Appendf(TEXT("  [ (Idx) | %-45s | % -40s | % -55s ]\n"), TEXT("Type"), TEXT("DataSource"), TEXT("Name"));
+		int32 Index = 0;
+		VisitSourceStructDescriptorInternal([&Index, &DebugString](const FPropertyBindingBindableStructDescriptor& Descriptor)
 		{
-			DebugString.Appendf(TEXT("  | %-45s | % -80s |\n"),
+			DebugString.Appendf(TEXT("  | (%3d) | %-45s | % -80s |\n"),
+				Index,
 				*GetNameSafe(Descriptor.Struct),
 				*Descriptor.ToString());
 
+			++Index;
 			return EVisitResult::Continue;
 		});
 	}
@@ -1110,14 +1113,16 @@ FString FPropertyBindingBindingCollection::DebugAsString() const
 	DebugString.Appendf(TEXT("\nCopyBatches (%d)\n"), CopyBatches.Num());
 	if (CopyBatches.Num())
 	{
-		DebugString.Appendf(TEXT("  [ %-45s | %-45s | %-8s [%-5s:%-5s[ | %-8s [%-5s:%-5s[ ]\n"),
+		DebugString.Appendf(TEXT("  [ (Idx) | %-45s | %-45s | %-8s [%-5s:%-5s[ | %-8s [%-5s:%-5s[ ]\n"),
 			TEXT("Target Type"), TEXT("Target Name"),
 			TEXT("Bindings"), TEXT("Beg"), TEXT("End"),
 			TEXT("ProFunc"), TEXT("Beg"), TEXT("End"));
-		for (const FPropertyBindingCopyInfoBatch& CopyBatch : CopyBatches)
+		for (int32 Index = 0; Index < CopyBatches.Num(); ++Index)
 		{
+			const FPropertyBindingCopyInfoBatch& CopyBatch = CopyBatches[Index];
 			const FPropertyBindingBindableStructDescriptor& Descriptor = CopyBatch.TargetStruct.Get();
-			DebugString.Appendf(TEXT("  | %-45s | %-45s | %8s [%5d:%-5d[ | %8s [%5d:%-5d[ |\n"),
+			DebugString.Appendf(TEXT("  | (%3d) | %-45s | %-45s | %8s [%5d:%-5d[ | %8s [%5d:%-5d[ |\n"),
+				Index,
 				Descriptor.Struct ? *Descriptor.Struct->GetName() : TEXT("null"),
 				*Descriptor.ToString(),
 				TEXT(""), CopyBatch.BindingsBegin.Get(), CopyBatch.BindingsEnd.Get(),
@@ -1129,13 +1134,16 @@ FString FPropertyBindingBindingCollection::DebugAsString() const
 	DebugString.Appendf(TEXT("\nPropertyPathBindings (%d)\n"), GetNumBindings());
 	if (GetNumBindings())
 	{
-		DebugString.Appendf(TEXT("  [ %-45s | %-45s | %-45s ]\n"),
+		DebugString.Appendf(TEXT("  [ (Idx) | %-45s | %-45s | %-45s ]\n"),
 			TEXT("Source"), TEXT("SrcPath"), TEXT("TargetPath"));
-		ForEachBinding([&DebugString](const FPropertyBindingBinding& PropertyBinding)
+		int32 Index = 0;
+		ForEachBinding([&Index, &DebugString](const FPropertyBindingBinding& PropertyBinding)
 		{
-			DebugString.Appendf(TEXT("  | %-45s | %-45s |\n"),
+			DebugString.Appendf(TEXT("  | (%3d) | %-45s | %-45s |\n"),
+				Index,
 				*PropertyBinding.GetSourcePath().ToString(),
 				*PropertyBinding.GetTargetPath().ToString());
+			++Index;
 		});
 	}
 
@@ -1143,13 +1151,15 @@ FString FPropertyBindingBindingCollection::DebugAsString() const
 	DebugString.Appendf(TEXT("\nPropertyCopies (%d)\n"), PropertyCopies.Num());
 	if (PropertyCopies.Num())
 	{
-		DebugString.Appendf(TEXT("  [ %-7s | %-4s | %-7s | %-10s | %-7s | %-4s | %-7s | %-10s | %-15s | %-20s | %-4s ]\n"),
+		DebugString.Appendf(TEXT("  [ (Idx) | %-7s | %-4s | %-7s | %-10s | %-7s | %-4s | %-7s | %-10s | %-15s | %-20s | %-4s ]\n"),
 			TEXT("Src Idx"), TEXT("Off."), TEXT("Next"), TEXT("Type"),
 			TEXT("Tgt Idx"), TEXT("Off."), TEXT("Next"), TEXT("Type"),
 			TEXT("Source"), TEXT("Copy Type"), TEXT("Size"));
-		for (const FPropertyBindingCopyInfo& PropertyCopy : PropertyCopies)
+		for (int32 Index = 0; Index < PropertyCopies.Num(); ++Index)
 		{
-			DebugString.Appendf(TEXT("  | %7d | %4d | %7d | %-10s | %7d | %4d | %7d | %-10s | %-20s | %4d |\n"),
+			const FPropertyBindingCopyInfo& PropertyCopy = PropertyCopies[Index];
+			DebugString.Appendf(TEXT("  | (%3d) | %7d | %4d | %7d | %-10s | %7d | %4d | %7d | %-10s | %-20s | %4d |\n"),
+						Index,
 						PropertyCopy.SourceIndirection.ArrayIndex.Get(),
 						PropertyCopy.SourceIndirection.Offset,
 						PropertyCopy.SourceIndirection.NextIndex.Get(),
@@ -1167,11 +1177,13 @@ FString FPropertyBindingBindingCollection::DebugAsString() const
 	DebugString.Appendf(TEXT("\nPropertyIndirections (%d)\n"), PropertyIndirections.Num());
 	if (PropertyIndirections.Num())
 	{
-		DebugString.Appendf(TEXT("[ %-4s | %-4s | %-4s | %-10s ] \n"),
+		DebugString.Appendf(TEXT("[ (Idx) | %-4s | %-4s | %-4s | %-10s ] \n"),
 			TEXT("Idx"), TEXT("Off."), TEXT("Next"), TEXT("Access Type"));
-		for (const FPropertyBindingPropertyIndirection& PropertyIndirection : PropertyIndirections)
+		for (int32 Index = 0; Index < PropertyIndirections.Num(); ++Index)
 		{
-			DebugString.Appendf(TEXT("  | %4d | %4d | %4d | %-10s |\n"),
+			const FPropertyBindingPropertyIndirection& PropertyIndirection = PropertyIndirections[Index];
+			DebugString.Appendf(TEXT("  | (%3d) | %4d | %4d | %4d | %-10s |\n"),
+				Index,
 				PropertyIndirection.ArrayIndex.Get(),
 				PropertyIndirection.Offset,
 				PropertyIndirection.NextIndex.Get(),
