StateTree: Prevents promotion of unsupported to a property bag. Fix the display and the binding action for a function.

diff --git a/Engine/Plugins/Runtime/PropertyBindingUtils/Source/PropertyBindingUtilsEditor/Private/PropertyBindingExtension.cpp b/Engine/Plugins/Runtime/PropertyBindingUtils/Source/PropertyBindingUtilsEditor/Private/PropertyBindingExtension.cpp
--- a/Engine/Plugins/Runtime/PropertyBindingUtils/Source/PropertyBindingUtilsEditor/Private/PropertyBindingExtension.cpp
+++ b/Engine/Plugins/Runtime/PropertyBindingUtils/Source/PropertyBindingUtilsEditor/Private/PropertyBindingExtension.cpp
@@ -357,27 +357,27 @@ void FCachedBindingData::AddBinding(TConstArrayView<FBindingChainElement> InBind
  	// Remove struct index.
 	const TConstArrayView<FBindingChainElement> SourceBindingChain(InBindingChain.begin() + 1, InBindingChain.Num() - 1);
 
+	// GetBindingDataViewByID can fail but we still need to call AddBindingInternal
 	FPropertyBindingDataView DataView;
-	if (BindingOwner->GetBindingDataViewByID(StructID, DataView))
-	{
-		// If SourceBindingChain is empty at this stage, it means that the binding points to the source struct itself.
-		FPropertyBindingPath SourcePath;
-		MakeStructPropertyPathFromBindingChain(StructID, SourceBindingChain, DataView, SourcePath);
-
-		UObject* BindingOwnerObject = WeakBindingsOwner.GetObject();
-		BindingOwnerObject->Modify();
+	BindingOwner->GetBindingDataViewByID(StructID, DataView);
 
-		// Allow derived classes to handle the bindings
-		const bool bBindingHandled = AddBindingInternal(BindableStruct, SourcePath, TargetPath);
-		if (bBindingHandled == false)
-		{
-			EditorBindings->AddBinding(SourcePath, TargetPath);
-		}
+	// If SourceBindingChain is empty at this stage, it means that the binding points to the source struct itself.
+	FPropertyBindingPath SourcePath;
+	MakeStructPropertyPathFromBindingChain(StructID, SourceBindingChain, DataView, SourcePath);
 
-		UpdateData();
+	UObject* BindingOwnerObject = WeakBindingsOwner.GetObject();
+	BindingOwnerObject->Modify();
 
-		BindingOwner->OnPropertyBindingChanged(SourcePath, TargetPath);
+	// Allow derived classes to handle the bindings
+	const bool bBindingHandled = AddBindingInternal(BindableStruct, SourcePath, TargetPath);
+	if (bBindingHandled == false)
+	{
+		EditorBindings->AddBinding(SourcePath, TargetPath);
 	}
+
+	UpdateData();
+
+	BindingOwner->OnPropertyBindingChanged(SourcePath, TargetPath);
 }
 
 bool FCachedBindingData::HasBinding(const FPropertyBindingBindingCollection::ESearchMode SearchMode) const
@@ -652,12 +652,7 @@ bool FCachedBindingData::CanBindToContextStruct(const UStruct* InStruct, int32 I
 {
 	ConditionallyUpdateData();
 
-	if (CanBindToContextStructInternal(InStruct, InStructIndex))
-	{
-		return ArePropertyAndContextStructCompatible(InStruct, PropertyHandle->GetProperty());
-	}
-
-	return false;
+	return CanBindToContextStructInternal(InStruct, InStructIndex);
 }
 
 void FCachedBindingData::GetSourceDataViewForNewBinding(TNotNull<IPropertyBindingBindingCollectionOwner*> InBindingsOwner, TConstStructView<FPropertyBindingBindableStructDescriptor> InDescriptor
@@ -832,6 +827,15 @@ bool FCachedBindingData::CanCreateParameter(const FPropertyBindingBindableStruct
 		return false;
 	}
 
+	// Is the type supported by the property bag
+	{
+		FPropertyBagPropertyDesc Desc = FPropertyBagPropertyDesc(Property->GetFName(), Property);
+		if (Desc.ValueType == EPropertyBagPropertyType::None)
+		{
+			return false;
+		}
+	}
+
 	IPropertyBindingBindingCollectionOwner* BindingsOwner = WeakBindingsOwner.Get();
 	if (BindingsOwner == nullptr)
 	{
@@ -857,7 +861,7 @@ bool FCachedBindingData::CanCreateParameter(const FPropertyBindingBindableStruct
 		return false;
 	}
 
-	// Allow derived c.lasses to push overrides
+	// Allow derived classes to push overrides
 	AddPropertyInfoOverride(*Property, OutPropertyInfoOverrides);
 
 	return true;
@@ -913,7 +917,7 @@ UStruct* FCachedBindingData::ResolveIndirection(TConstArrayView<FBindingChainEle
 
 bool FCachedBindingData::CanBindToContextStructInternal(const UStruct* InStruct, int32 InStructIndex)
 {
-	return true;
+	return ArePropertyAndContextStructCompatible(InStruct, PropertyHandle->GetProperty());
 }
 
 bool FCachedBindingData::GetPropertyFunctionText(FConstStructView InPropertyFunctionStructView, FText& OutText) const
@@ -1478,6 +1482,15 @@ bool FPropertyBindingExtension::CanPromoteToParameter(const TSharedPtr<IProperty
 		}
 	}
 
+	// Is the type supported by the property bag
+	{
+		FPropertyBagPropertyDesc Desc = FPropertyBagPropertyDesc(Property->GetFName(), Property);
+		if (Desc.ValueType == EPropertyBagPropertyType::None)
+		{
+			return false;
+		}
+	}
+
 	return true;
 }
 
diff --git a/Engine/Plugins/Runtime/StateTree/Source/StateTreeEditorModule/Private/Customizations/StateTreeBindingExtension.cpp b/Engine/Plugins/Runtime/StateTree/Source/StateTreeEditorModule/Private/Customizations/StateTreeBindingExtension.cpp
--- a/Engine/Plugins/Runtime/StateTree/Source/StateTreeEditorModule/Private/Customizations/StateTreeBindingExtension.cpp
+++ b/Engine/Plugins/Runtime/StateTree/Source/StateTreeEditorModule/Private/Customizations/StateTreeBindingExtension.cpp
@@ -302,7 +302,7 @@ struct FStateTreeCachedBindingData : UE::PropertyBinding::FCachedBindingData, TS
 			}
 		}
 
-		return true;
+		return Super::CanBindToContextStructInternal(InStruct, InStructIndex);
 	}
 
 	virtual bool CanAcceptPropertyOrChildrenInternal(const FProperty& SourceProperty, TConstArrayView<FBindingChainElement, int> InBindingChain) override
