PropertyBag: Prevents crashes when the user requests an unsupported type.

diff --git a/Engine/Plugins/Runtime/PropertyBindingUtils/Source/PropertyBindingUtilsEditor/Private/PropertyBindingExtension.cpp b/Engine/Plugins/Runtime/PropertyBindingUtils/Source/PropertyBindingUtilsEditor/Private/PropertyBindingExtension.cpp
--- a/Engine/Plugins/Runtime/PropertyBindingUtils/Source/PropertyBindingUtilsEditor/Private/PropertyBindingExtension.cpp
+++ b/Engine/Plugins/Runtime/PropertyBindingUtils/Source/PropertyBindingUtilsEditor/Private/PropertyBindingExtension.cpp
@@ -828,12 +828,10 @@ bool FCachedBindingData::CanCreateParameter(const FPropertyBindingBindableStruct
 	}
 
 	// Is the type supported by the property bag
+	const FPropertyBagPropertyDesc Desc = FPropertyBagPropertyDesc(Property->GetFName(), Property);
+	if (Desc.ValueType == EPropertyBagPropertyType::None)
 	{
-		FPropertyBagPropertyDesc Desc = FPropertyBagPropertyDesc(Property->GetFName(), Property);
-		if (Desc.ValueType == EPropertyBagPropertyType::None)
-		{
-			return false;
-		}
+		return false;
 	}
 
 	IPropertyBindingBindingCollectionOwner* BindingsOwner = WeakBindingsOwner.Get();
@@ -1483,12 +1481,10 @@ bool FPropertyBindingExtension::CanPromoteToParameter(const TSharedPtr<IProperty
 	}
 
 	// Is the type supported by the property bag
+	const FPropertyBagPropertyDesc Desc = FPropertyBagPropertyDesc(Property->GetFName(), Property);
+	if (Desc.ValueType == EPropertyBagPropertyType::None)
 	{
-		FPropertyBagPropertyDesc Desc = FPropertyBagPropertyDesc(Property->GetFName(), Property);
-		if (Desc.ValueType == EPropertyBagPropertyType::None)
-		{
-			return false;
-		}
+		return false;
 	}
 
 	return true;
