[PropertyBindings] Added extra braces to fix clang compilation

diff --git a/Engine/Plugins/Runtime/PropertyBindingUtils/Source/PropertyBindingUtils/Private/PropertyBindingBindingCollection.cpp b/Engine/Plugins/Runtime/PropertyBindingUtils/Source/PropertyBindingUtils/Private/PropertyBindingBindingCollection.cpp
--- a/Engine/Plugins/Runtime/PropertyBindingUtils/Source/PropertyBindingUtils/Private/PropertyBindingBindingCollection.cpp
+++ b/Engine/Plugins/Runtime/PropertyBindingUtils/Source/PropertyBindingUtils/Private/PropertyBindingBindingCollection.cpp
@@ -26,7 +26,7 @@
 #endif // WITH_EDITOR
 
 #define PROPERTY_BINDINGS_CLOG(Condition, LogOwner, Verbosity, Format, ...) \
-	if (Condition) PROPERTY_BINDINGS_LOG(LogOwner, Verbosity, Format, ##__VA_ARGS__);
+	if (Condition) { PROPERTY_BINDINGS_LOG(LogOwner, Verbosity, Format, ##__VA_ARGS__); }
 
 const UObject* FPropertyBindingBindingCollection::GetLogOwner() const
 {
