[PropertyBindings] prevent triggering a check when performing copy on outdated bindings (out of bounds index)

diff --git a/Engine/Plugins/Runtime/PropertyBindingUtils/Source/PropertyBindingUtils/Private/PropertyBindingBindingCollection.cpp b/Engine/Plugins/Runtime/PropertyBindingUtils/Source/PropertyBindingUtils/Private/PropertyBindingBindingCollection.cpp
--- a/Engine/Plugins/Runtime/PropertyBindingUtils/Source/PropertyBindingUtils/Private/PropertyBindingBindingCollection.cpp
+++ b/Engine/Plugins/Runtime/PropertyBindingUtils/Source/PropertyBindingUtils/Private/PropertyBindingBindingCollection.cpp
@@ -883,13 +883,25 @@ uint8* FPropertyBindingBindingCollection::GetAddress(FPropertyBindingDataView In
 	return (uint8*)Address;
 }
 
-void FPropertyBindingBindingCollection::PerformCopy(const FPropertyBindingCopyInfo& Copy, uint8* SourceAddress, uint8* TargetAddress) const
+bool FPropertyBindingBindingCollection::PerformCopy(const FPropertyBindingCopyInfo& Copy, uint8* SourceAddress, uint8* TargetAddress) const
 {
-	// Source property can be null
-	check(SourceAddress || Copy.Type == EPropertyCopyType::CopyObject);
-	check(Copy.TargetLeafProperty);
-	check(TargetAddress);
+	// 'SourceAddress' can only be nullptr for object copy (e.g., EPropertyCopyType::CopyObject)
+	// Otherwise we simply fail the copy since it might be possible to get outdated bindings (e.g., out of bound array index)
+	if (SourceAddress == nullptr && Copy.Type != EPropertyCopyType::CopyObject)
+	{
+		PROPERTY_BINDINGS_LOG(GetLogOwner(), Verbose, TEXT("%hs skipped: invalid source address for copy type '%s'."), __FUNCTION__, *UEnum::GetValueAsString(Copy.Type));
+		return false;
+	}
+
+	// Target address is always required
+	if (TargetAddress == nullptr)
+	{
+		PROPERTY_BINDINGS_LOG(GetLogOwner(), Verbose, TEXT("%hs skipped: invalid target address for copy type '%s'."), __FUNCTION__, *UEnum::GetValueAsString(Copy.Type));
+		return false;
+	}
 
+	check(Copy.TargetLeafProperty);
+	
 	switch (Copy.Type)
 	{
 	case EPropertyCopyType::CopyPlain:
@@ -1030,6 +1042,8 @@ void FPropertyBindingBindingCollection::PerformCopy(const FPropertyBindingCopyIn
 			*StaticEnum<EPropertyCopyType>()->GetValueAsString(Copy.Type), *Copy.SourceLeafProperty->GetNameCPP(), *Copy.TargetLeafProperty->GetNameCPP());
 		break;
 	}
+
+	return true;
 }
 
 bool FPropertyBindingBindingCollection::CopyProperty(const FPropertyBindingCopyInfo& Copy, FPropertyBindingDataView SourceStructView, FPropertyBindingDataView TargetStructView) const
@@ -1056,15 +1070,7 @@ bool FPropertyBindingBindingCollection::CopyProperty(const FPropertyBindingCopyI
 		uint8* SourceAddress = GetAddress(SourceStructView, Copy.SourceIndirection, Copy.SourceLeafProperty);
 		uint8* TargetAddress = GetAddress(TargetStructView, Copy.TargetIndirection, Copy.TargetLeafProperty);
 
-		// 'SourceAddress' can be nullptr for some types of copy (e.g., EPropertyCopyType::CopyObject)
-		if (TargetAddress != nullptr)
-		{
-			PerformCopy(Copy, SourceAddress, TargetAddress);
-		}
-		else
-		{
-			bResult = false;
-		}
+		return PerformCopy(Copy, SourceAddress, TargetAddress);
 	}
 	else
 	{
diff --git a/Engine/Plugins/Runtime/PropertyBindingUtils/Source/PropertyBindingUtils/Public/PropertyBindingBindingCollection.h b/Engine/Plugins/Runtime/PropertyBindingUtils/Source/PropertyBindingUtils/Public/PropertyBindingBindingCollection.h
--- a/Engine/Plugins/Runtime/PropertyBindingUtils/Source/PropertyBindingUtils/Public/PropertyBindingBindingCollection.h
+++ b/Engine/Plugins/Runtime/PropertyBindingUtils/Source/PropertyBindingUtils/Public/PropertyBindingBindingCollection.h
@@ -305,7 +305,7 @@ protected:
 
 	[[nodiscard]] PROPERTYBINDINGUTILS_API bool ResolvePath(const UStruct* Struct, const FPropertyBindingPath& Path, FPropertyBindingPropertyIndirection& OutFirstIndirection, FPropertyBindingPathIndirection& OutLeafIndirection);
 
-	PROPERTYBINDINGUTILS_API void PerformCopy(const FPropertyBindingCopyInfo& Copy, uint8* SourceAddress, uint8* TargetAddress) const;
+	PROPERTYBINDINGUTILS_API bool PerformCopy(const FPropertyBindingCopyInfo& Copy, uint8* SourceAddress, uint8* TargetAddress) const;
 	PROPERTYBINDINGUTILS_API void PerformResetObjects(const FPropertyBindingCopyInfo& Copy, uint8* TargetAddress) const;
 	PROPERTYBINDINGUTILS_API uint8* GetAddress(FPropertyBindingDataView InStructView, const FPropertyBindingPropertyIndirection& FirstIndirection, const FProperty* LeafProperty) const;
 
