[PropertyBindings] do not apply extension on properties that can't be used in bindings

diff --git a/Engine/Plugins/Runtime/PropertyBindingUtils/Source/PropertyBindingUtilsEditor/Private/PropertyBindingExtension.cpp b/Engine/Plugins/Runtime/PropertyBindingUtils/Source/PropertyBindingUtilsEditor/Private/PropertyBindingExtension.cpp
--- a/Engine/Plugins/Runtime/PropertyBindingUtils/Source/PropertyBindingUtilsEditor/Private/PropertyBindingExtension.cpp
+++ b/Engine/Plugins/Runtime/PropertyBindingUtils/Source/PropertyBindingUtilsEditor/Private/PropertyBindingExtension.cpp
@@ -1088,6 +1088,13 @@ bool IsPropertyBindable(const FProperty& Property)
 		return false;
 	}
 
+	if (const bool bIsEditConst = Property.HasAnyPropertyFlags(CPF_EditConst))
+	{
+		UE_LOG(LogPropertyBindingUtils, Verbose, TEXT("Property %s is not bindable because it's readonly in the editor"),
+			*Property.GetName());
+		return false;
+	}
+
 	const bool bPrivateOrProtected = !Property.HasAnyPropertyFlags(CPF_NativeAccessSpecifierPrivate | CPF_NativeAccessSpecifierProtected);
 	const bool bPrivateButBlueprintAccessible = Property.GetBoolMetaData(FBlueprintMetadata::MD_AllowPrivateAccess);
 	if (!bPrivateOrProtected && !bPrivateButBlueprintAccessible)
@@ -1175,6 +1182,14 @@ void FPropertyBindingExtension::ExtendWidgetRow(FDetailWidgetRow& InWidgetRow, c
 		return;
 	}
 
+	if (const FProperty* Property = InPropertyHandle->GetProperty())
+	{
+		if (!UE::PropertyBinding::IsPropertyBindable(*Property))
+		{
+			return;
+		}
+	}
+
 	using namespace UE::PropertyBinding;
 
 	IPropertyAccessEditor& PropertyAccessEditor = IModularFeatures::Get().GetModularFeature<IPropertyAccessEditor>("PropertyAccessEditor");
