PropertyBinding: Add metadata to disable the Promote to Parameter option in the context menu.

diff --git a/Engine/Plugins/Runtime/PropertyBindingUtils/Source/PropertyBindingUtilsEditor/Private/PropertyBindingExtension.cpp b/Engine/Plugins/Runtime/PropertyBindingUtils/Source/PropertyBindingUtilsEditor/Private/PropertyBindingExtension.cpp
--- a/Engine/Plugins/Runtime/PropertyBindingUtils/Source/PropertyBindingUtilsEditor/Private/PropertyBindingExtension.cpp
+++ b/Engine/Plugins/Runtime/PropertyBindingUtils/Source/PropertyBindingUtilsEditor/Private/PropertyBindingExtension.cpp
@@ -247,6 +247,37 @@ TSharedPtr<const IPropertyHandle> MakeStructPropertyPathFromPropertyHandle(
 	return BindablePropertyHandle;
 }
 
+bool HasMetaData(FName MetaData, const TSharedRef<const IPropertyHandle>& InPropertyHandle)
+{
+	TSharedPtr<const IPropertyHandle> CurrentPropertyHandle = InPropertyHandle;
+	while (CurrentPropertyHandle.IsValid())
+	{
+		if (const FProperty* MetaDataProperty = CurrentPropertyHandle->GetMetaDataProperty())
+		{
+			if (const FStructProperty* StructProperty = CastField<const FStructProperty>(MetaDataProperty))
+			{
+				if (StructProperty->Struct->HasMetaData(MetaData))
+				{
+					return true;
+				}
+			}
+			else if (const FObjectPropertyBase* ObjectProperty = CastField<const FObjectPropertyBase>(MetaDataProperty))
+			{
+				if (ObjectProperty->PropertyClass->HasMetaData(MetaData))
+				{
+					return true;
+				}
+			}
+		}
+		if (CurrentPropertyHandle->HasMetaData(MetaData))
+		{
+			return true;
+		}
+		CurrentPropertyHandle = CurrentPropertyHandle->GetParentHandle();
+	}
+	return false;
+}
+
 FMenuSectionHelper::FMenuSectionHelper(FMenuBuilder& MenuBuilder): MenuBuilder(MenuBuilder)
 {
 }
@@ -1118,31 +1149,9 @@ bool FPropertyBindingExtension::IsPropertyExtendable(const UClass* InObjectClass
 	}
 
 	// Does the container or property support bindings
-	TSharedPtr<const IPropertyHandle> CurrentPropertyHandle = InPropertyHandle.AsShared();
-	while (CurrentPropertyHandle.IsValid())
+	if (UE::PropertyBinding::HasMetaData(UE::PropertyBinding::MetaDataNoBindingName, InPropertyHandle.AsShared()))
 	{
-		if (const FProperty* MetaDataProperty = CurrentPropertyHandle->GetMetaDataProperty())
-		{
-			if (const FStructProperty* StructProperty = CastField<const FStructProperty>(MetaDataProperty))
-			{
-				if (StructProperty->Struct->HasMetaData(UE::PropertyBinding::MetaDataNoBindingName))
-				{
-					return false;
-				}
-			}
-			else if (const FObjectPropertyBase* ObjectProperty = CastField<const FObjectPropertyBase>(MetaDataProperty))
-			{
-				if (ObjectProperty->PropertyClass->HasMetaData(UE::PropertyBinding::MetaDataNoBindingName))
-				{
-					return false;
-				}
-			}
-		}
-		if (CurrentPropertyHandle->HasMetaData(UE::PropertyBinding::MetaDataNoBindingName))
-		{
-			return false;
-		}
-		CurrentPropertyHandle = CurrentPropertyHandle->GetParentHandle();
+		return false;
 	}
 
 	TArray<UObject*> OuterObjects;
@@ -1515,6 +1524,11 @@ bool FPropertyBindingExtension::CanPromoteToParameter(const TSharedPtr<IProperty
 		return false;
 	}
 
+	if (UE::PropertyBinding::HasMetaData(UE::PropertyBinding::MetaDataNoPromoteToParameter, InPropertyHandle.ToSharedRef()))
+	{
+		return false;
+	}
+
 	return true;
 }
 
diff --git a/Engine/Plugins/Runtime/PropertyBindingUtils/Source/PropertyBindingUtilsEditor/Public/PropertyBindingExtension.h b/Engine/Plugins/Runtime/PropertyBindingUtils/Source/PropertyBindingUtilsEditor/Public/PropertyBindingExtension.h
--- a/Engine/Plugins/Runtime/PropertyBindingUtils/Source/PropertyBindingUtilsEditor/Public/PropertyBindingExtension.h
+++ b/Engine/Plugins/Runtime/PropertyBindingUtils/Source/PropertyBindingUtilsEditor/Public/PropertyBindingExtension.h
@@ -34,16 +34,19 @@ struct FPropertyInfoOverride
 
 namespace Meta::Private
 {
-// valid keywords for the UCLASS UPROPERTY and USTRUCT macro
+// valid keywords for the UCLASS, UPROPERTY and USTRUCT macro
 enum
 {
 	// [ClassMetadata] [PropertyMetadata] [StructMetadata] The property cannot be bound to (target of a binding).
 	NoBinding,
+	// [ClassMetadata] [PropertyMetadata] [StructMetadata] The property cannot be promoted to parameters.
+	NoPromoteToParameter,
 };
 } // namespace Meta::Private
 
 const FLazyName MetaDataStructIDName("StructIDForPropertyBinding");
 const FLazyName MetaDataNoBindingName("NoBinding");
+const FLazyName MetaDataNoPromoteToParameter("NoPromoteToParameter");
 
 PROPERTYBINDINGUTILSEDITOR_API extern FText GetPropertyTypeText(const FProperty* Property);
 
diff --git a/Engine/Plugins/Runtime/StateTree/Source/StateTreeEditorModule/Private/Customizations/StateTreeBindingExtension.cpp b/Engine/Plugins/Runtime/StateTree/Source/StateTreeEditorModule/Private/Customizations/StateTreeBindingExtension.cpp
--- a/Engine/Plugins/Runtime/StateTree/Source/StateTreeEditorModule/Private/Customizations/StateTreeBindingExtension.cpp
+++ b/Engine/Plugins/Runtime/StateTree/Source/StateTreeEditorModule/Private/Customizations/StateTreeBindingExtension.cpp
@@ -420,11 +420,8 @@ struct FStateTreeCachedBindingData : UE::PropertyBinding::FCachedBindingData, TS
 		}
 		else if (TargetStructProperty && IsDelegateListenerProperty(TargetStructProperty))
 		{
-			if (const FStructProperty* SourceStructProperty = CastField<FStructProperty>(InSourceProperty))
-			{
-				bOutAreCompatible = SourceStructProperty->Struct == FStateTreeDelegateDispatcher::StaticStruct();
-				return true;
-			}
+			bOutAreCompatible = IsDelegateDispatcherProperty(InSourceProperty);
+			return true;
 		}
 
 		return false;
diff --git a/Engine/Plugins/Runtime/StateTree/Source/StateTreeModule/Public/StateTreeDelegate.h b/Engine/Plugins/Runtime/StateTree/Source/StateTreeModule/Public/StateTreeDelegate.h
--- a/Engine/Plugins/Runtime/StateTree/Source/StateTreeModule/Public/StateTreeDelegate.h
+++ b/Engine/Plugins/Runtime/StateTree/Source/StateTreeModule/Public/StateTreeDelegate.h
@@ -24,7 +24,7 @@
  /**
   * StateTree's delegate dispatcher.
   */
-USTRUCT(BlueprintType, meta=(NoBinding))
+USTRUCT(BlueprintType, meta = (NoBinding))
 struct STATETREEMODULE_API FStateTreeDelegateDispatcher
 {
 	GENERATED_BODY()
@@ -50,7 +50,7 @@ private:
  * The receiver of a delegate binding.
  * Can be bound in the editor to a delegate dispatcher.
  */
-USTRUCT(BlueprintType)
+USTRUCT(BlueprintType, meta = (NoPromoteToParameter))
 struct STATETREEMODULE_API FStateTreeDelegateListener
 {
 	GENERATED_BODY()
