[PropertyBindingCollection] Fix for an array modification during ranged-for iteration

diff --git a/Engine/Plugins/Runtime/PropertyBindingUtils/Source/PropertyBindingUtils/Private/PropertyBindingBindingCollection.cpp b/Engine/Plugins/Runtime/PropertyBindingUtils/Source/PropertyBindingUtils/Private/PropertyBindingBindingCollection.cpp
--- a/Engine/Plugins/Runtime/PropertyBindingUtils/Source/PropertyBindingUtils/Private/PropertyBindingBindingCollection.cpp
+++ b/Engine/Plugins/Runtime/PropertyBindingUtils/Source/PropertyBindingUtils/Private/PropertyBindingBindingCollection.cpp
@@ -100,17 +100,19 @@ void FPropertyBindingBindingCollection::RemoveBindings(TFunctionRef<bool(FProper
 void FPropertyBindingBindingCollection::CopyBindings(const FGuid InFromStructID, const FGuid InToStructID)
 {
 	// Copy all bindings that target "InFromStructID" and retarget them to "InToStructID".
-	const int32 NumBindingsToVisit = GetNumBindings();
-	int32 NumVisited = 0;
-	VisitBindings([&NumVisited, NumBindingsToVisit, InFromStructID, InToStructID, this](const FPropertyBindingBinding& Binding)
+	TArray<TPair<FPropertyBindingPath, FPropertyBindingPath>, TInlineAllocator<8>> BindingsToCopy;
+	ForEachBinding([&BindingsToCopy, InFromStructID, InToStructID](const FPropertyBindingBinding& Binding)
 	{
-		NumVisited++;
 		if (Binding.GetTargetPath().GetStructID() == InFromStructID)
 		{
-			AddBindingInternal(Binding.GetSourcePath(), FPropertyBindingPath(InToStructID, Binding.GetTargetPath().GetSegments()));
+			BindingsToCopy.Add({Binding.GetSourcePath(), FPropertyBindingPath(InToStructID, Binding.GetTargetPath().GetSegments())});
 		}
-		return NumVisited == NumBindingsToVisit ? EVisitResult::Break : EVisitResult::Continue;
 	});
+
+	for (const TPair<FPropertyBindingPath, FPropertyBindingPath>& BindingToCopy : BindingsToCopy)
+	{
+		AddBindingInternal(BindingToCopy.Key, BindingToCopy.Value);
+	}
 }
 
 bool FPropertyBindingBindingCollection::HasBinding(const FPropertyBindingPath& InTargetPath, ESearchMode RemoveMode) const
