PropertyBindings: Use redirector when resolving binding path to solve issues with renamed BP property.

diff --git a/Engine/Plugins/Runtime/PropertyBindingUtils/Source/PropertyBindingUtilsEditor/Private/PropertyBindingExtension.cpp b/Engine/Plugins/Runtime/PropertyBindingUtils/Source/PropertyBindingUtilsEditor/Private/PropertyBindingExtension.cpp
--- a/Engine/Plugins/Runtime/PropertyBindingUtils/Source/PropertyBindingUtilsEditor/Private/PropertyBindingExtension.cpp
+++ b/Engine/Plugins/Runtime/PropertyBindingUtils/Source/PropertyBindingUtilsEditor/Private/PropertyBindingExtension.cpp
@@ -563,8 +563,8 @@ void FCachedBindingData::UpdateData()
 				// Resolve source and target properties.
 				// Source path can be empty, when the binding binds directly to a context struct/class.
 				// Target path must always point to a valid property (at least one indirection).
-				if (SourcePath.ResolveIndirectionsWithValue(SourceDataView, SourceIndirections)
-					&& TargetPath.ResolveIndirectionsWithValue(TargetDataView, TargetIndirections)
+				if (SourcePath.ResolveIndirectionsWithValue(SourceDataView, SourceIndirections, /*error*/nullptr, /*bHandleRedirects*/true)
+					&& TargetPath.ResolveIndirectionsWithValue(TargetDataView, TargetIndirections, /*error*/nullptr, /*bHandleRedirects*/true)
 					&& !TargetIndirections.IsEmpty())
 				{
 					const FPropertyBindingPathIndirection TargetLeafIndirection = TargetIndirections.Last();
