StateTree: Fixed copy binding for function bindings.

diff --git a/Engine/Plugins/Runtime/PropertyBindingUtils/Source/PropertyBindingUtils/Private/PropertyBindingBindingCollection.cpp b/Engine/Plugins/Runtime/PropertyBindingUtils/Source/PropertyBindingUtils/Private/PropertyBindingBindingCollection.cpp
--- a/Engine/Plugins/Runtime/PropertyBindingUtils/Source/PropertyBindingUtils/Private/PropertyBindingBindingCollection.cpp
+++ b/Engine/Plugins/Runtime/PropertyBindingUtils/Source/PropertyBindingUtils/Private/PropertyBindingBindingCollection.cpp
@@ -99,19 +99,33 @@ void FPropertyBindingBindingCollection::RemoveBindings(TFunctionRef<bool(FProper
 
 void FPropertyBindingBindingCollection::CopyBindings(const FGuid InFromStructID, const FGuid InToStructID)
 {
-	// Copy all bindings that target "InFromStructID" and retarget them to "InToStructID".
-	TArray<TPair<FPropertyBindingPath, FPropertyBindingPath>, TInlineAllocator<8>> BindingsToCopy;
-	ForEachBinding([&BindingsToCopy, InFromStructID, InToStructID](const FPropertyBindingBinding& Binding)
-	{
-		if (Binding.GetTargetPath().GetStructID() == InFromStructID)
+	CopyBindingsInternal(InFromStructID, InToStructID);
+}
+
+void FPropertyBindingBindingCollection::CopyBindingsInternal(const FGuid InFromStructID, const FGuid InToStructID)
+{
+	CopyBindingsImplementation(InFromStructID, InToStructID, [](const FPropertyBindingBinding&){ return true; });
+}
+
+void FPropertyBindingBindingCollection::CopyBindingsImplementation(const FGuid InFromStructID, const FGuid InToStructID, TFunctionRef<bool(const FPropertyBindingBinding& Binding)> CanCopy)
+{
+	// Find the StructID and copy the binding. If we find functions, then copy the function and copy the functions binding recursively.
+	TArray<TTuple<FPropertyBindingPath, FPropertyBindingPath>, TInlineAllocator<8>> BindingsToCopy;
+	ForEachBinding([&InFromStructID, &InToStructID, &BindingsToCopy, &CanCopy](const FPropertyBindingBinding& Binding)
 		{
-			BindingsToCopy.Add({Binding.GetSourcePath(), FPropertyBindingPath(InToStructID, Binding.GetTargetPath().GetSegments())});
-		}
-	});
+			if (Binding.GetTargetPath().GetStructID() == InFromStructID)
+			{
+				if (CanCopy(Binding))
+				{
+					BindingsToCopy.Emplace(Binding.GetSourcePath(), FPropertyBindingPath(InToStructID, Binding.GetTargetPath().GetSegments()));
+				}
+			}
+		});
 
-	for (const TPair<FPropertyBindingPath, FPropertyBindingPath>& BindingToCopy : BindingsToCopy)
+	// Copy all bindings that target "FromStructID" and retarget them to "ToStructID".
+	for (const TTuple<FPropertyBindingPath, FPropertyBindingPath>& BindingToCopy : BindingsToCopy)
 	{
-		AddBindingInternal(BindingToCopy.Key, BindingToCopy.Value);
+		AddBindingInternal(BindingToCopy.Get<0>(), BindingToCopy.Get<1>());
 	}
 }
 
diff --git a/Engine/Plugins/Runtime/PropertyBindingUtils/Source/PropertyBindingUtils/Public/PropertyBindingBindingCollection.h b/Engine/Plugins/Runtime/PropertyBindingUtils/Source/PropertyBindingUtils/Public/PropertyBindingBindingCollection.h
--- a/Engine/Plugins/Runtime/PropertyBindingUtils/Source/PropertyBindingUtils/Public/PropertyBindingBindingCollection.h
+++ b/Engine/Plugins/Runtime/PropertyBindingUtils/Source/PropertyBindingUtils/Public/PropertyBindingBindingCollection.h
@@ -120,13 +120,15 @@ struct FPropertyBindingBindingCollection
 	PROPERTYBINDINGUTILS_API void GetBindingsFor(const FGuid InStructID, TArray<const FPropertyBindingBinding*>& OutBindings) const;
 
 protected:
-
 	virtual FPropertyBindingBinding* AddBindingInternal(const FPropertyBindingPath& InSourcePath, const FPropertyBindingPath& InTargetPath) PURE_VIRTUAL(FPropertyBindingBindingCollection::AddBindingInternal, return nullptr; );
+	PROPERTYBINDINGUTILS_API virtual void CopyBindingsInternal(const FGuid InFromStructID, const FGuid InToStructID);
 	virtual void RemoveBindingsInternal(TFunctionRef<bool(FPropertyBindingBinding&)> InPredicate) PURE_VIRTUAL(FPropertyBindingBindingCollection::RemoveBindingsInternal, );
 	virtual bool HasBindingInternal(TFunctionRef<bool(const FPropertyBindingBinding&)> InPredicate) const PURE_VIRTUAL(FPropertyBindingBindingCollection::HasBindingInternal, return false; );
 	virtual const FPropertyBindingBinding* FindBindingInternal(TFunctionRef<bool(const FPropertyBindingBinding&)> InPredicate) const PURE_VIRTUAL(FPropertyBindingBindingCollection::FindBindingInternal, return nullptr; );
-#endif // WITH_EDITOR
 
+	/** Copies property bindings from an existing struct to another. */
+	PROPERTYBINDINGUTILS_API void CopyBindingsImplementation(const FGuid InFromStructID, const FGuid InToStructID, TFunctionRef<bool(const FPropertyBindingBinding& Binding)> CanCopy);
+#endif // WITH_EDITOR
 	const UObject* GetLogOwner() const;
 
 public:
diff --git a/Engine/Plugins/Runtime/StateTree/Source/StateTreeEditorModule/Private/StateTreeEditorPropertyBindings.cpp b/Engine/Plugins/Runtime/StateTree/Source/StateTreeEditorModule/Private/StateTreeEditorPropertyBindings.cpp
--- a/Engine/Plugins/Runtime/StateTree/Source/StateTreeEditorModule/Private/StateTreeEditorPropertyBindings.cpp
+++ b/Engine/Plugins/Runtime/StateTree/Source/StateTreeEditorModule/Private/StateTreeEditorPropertyBindings.cpp
@@ -46,6 +46,42 @@ FPropertyBindingBinding* FStateTreeEditorPropertyBindings::AddBindingInternal(co
 	return &PropertyBindings.Emplace_GetRef(InSourcePath, InTargetPath);
 }
 
+void FStateTreeEditorPropertyBindings::CopyBindingsInternal(const FGuid InFromStructID, const FGuid InToStructID)
+{
+	// Find the StructID and copy the binding. If we find functions, then copy the function and copy the functions binding recursively.
+	TArray<TTuple<FGuid, FGuid>, TInlineAllocator<8>> TargetIDs;
+	TargetIDs.Emplace(InFromStructID, InToStructID);
+
+	for (int32 Index = 0; Index < TargetIDs.Num(); ++Index)
+	{
+		const FGuid& FromStructID = TargetIDs[Index].Get<0>();
+		const FGuid& ToStructID = TargetIDs[Index].Get<1>();
+
+		TArray<TTuple<const UScriptStruct*, FPropertyBindingPath, FPropertyBindingPath>, TInlineAllocator<8>> FunctionBindingsToCopy;
+		CopyBindingsImplementation(FromStructID, ToStructID, [ToStructID, &FunctionBindingsToCopy](const FPropertyBindingBinding& Binding)
+			{
+				const bool bIsFunctionBinding = Binding.GetPropertyFunctionNode().IsValid();
+				if (bIsFunctionBinding)
+				{
+					const FStateTreeEditorNode* EditorNode = Binding.GetPropertyFunctionNode().GetPtr<const FStateTreeEditorNode>();
+					if (ensure(EditorNode))
+					{
+						FunctionBindingsToCopy.Emplace(EditorNode->Node.GetScriptStruct(), Binding.GetSourcePath(), FPropertyBindingPath(ToStructID, Binding.GetTargetPath().GetSegments()));
+					}
+				}
+				return !bIsFunctionBinding;
+			});
+
+		// Copy all functions bindings that target "FromStructID" and retarget them to "ToStructID".
+		for (const TTuple<const UScriptStruct*, FPropertyBindingPath, FPropertyBindingPath>& FunctionBindingToCopy : FunctionBindingsToCopy)
+		{
+			const FGuid NewID = AddFunctionBinding(FunctionBindingToCopy.Get<0>(), FunctionBindingToCopy.Get<1>().GetSegments(), FunctionBindingToCopy.Get<2>()).GetStructID();
+			// Adds the new function binding to the list of copy bindings to look for.
+			TargetIDs.AddUnique(TTuple<FGuid, FGuid>{FunctionBindingToCopy.Get<1>().GetStructID(), NewID});
+		}
+	}
+}
+
 void FStateTreeEditorPropertyBindings::RemoveBindingsInternal(TFunctionRef<bool(FPropertyBindingBinding&)> InPredicate)
 {
 	PropertyBindings.RemoveAllSwap(InPredicate);
diff --git a/Engine/Plugins/Runtime/StateTree/Source/StateTreeEditorModule/Public/StateTreeEditorPropertyBindings.h b/Engine/Plugins/Runtime/StateTree/Source/StateTreeEditorModule/Public/StateTreeEditorPropertyBindings.h
--- a/Engine/Plugins/Runtime/StateTree/Source/StateTreeEditorModule/Public/StateTreeEditorPropertyBindings.h
+++ b/Engine/Plugins/Runtime/StateTree/Source/StateTreeEditorModule/Public/StateTreeEditorPropertyBindings.h
@@ -148,6 +148,7 @@ PRAGMA_ENABLE_DEPRECATION_WARNINGS
 protected:
 	//~ Begin FPropertyBindingBindingCollection overrides
 	virtual FPropertyBindingBinding* AddBindingInternal(const FPropertyBindingPath& InSourcePath, const FPropertyBindingPath& InTargetPath) override;
+	virtual void CopyBindingsInternal(const FGuid InFromStructID, const FGuid InToStructID) override;
 	virtual void RemoveBindingsInternal(TFunctionRef<bool(FPropertyBindingBinding&)> InPredicate) override;
 	virtual bool HasBindingInternal(TFunctionRef<bool(const FPropertyBindingBinding&)> InPredicate) const override;
 	virtual const FPropertyBindingBinding* FindBindingInternal(TFunctionRef<bool(const FPropertyBindingBinding&)> InPredicate) const override;
