[PropertyBindings] FPropertyBindingExtension::ExtendWidgetRow should only rely on IsPropertyExtendable to decide which properties can use bindings. Call to IsPropertyBindable was added in recent refactoring to reduce undesired bindings in SmartObjectDefinition but that was not the right fix.

diff --git a/Engine/Plugins/Runtime/PropertyBindingUtils/Source/PropertyBindingUtilsEditor/Private/PropertyBindingExtension.cpp b/Engine/Plugins/Runtime/PropertyBindingUtils/Source/PropertyBindingUtilsEditor/Private/PropertyBindingExtension.cpp
--- a/Engine/Plugins/Runtime/PropertyBindingUtils/Source/PropertyBindingUtilsEditor/Private/PropertyBindingExtension.cpp
+++ b/Engine/Plugins/Runtime/PropertyBindingUtils/Source/PropertyBindingUtilsEditor/Private/PropertyBindingExtension.cpp
@@ -1191,14 +1191,6 @@ void FPropertyBindingExtension::ExtendWidgetRow(FDetailWidgetRow& InWidgetRow, c
 		return;
 	}
 
-	if (const FProperty* Property = InPropertyHandle->GetProperty())
-	{
-		if (!UE::PropertyBinding::IsPropertyBindable(*Property))
-		{
-			return;
-		}
-	}
-
 	using namespace UE::PropertyBinding;
 
 	IPropertyAccessEditor& PropertyAccessEditor = IModularFeatures::Get().GetModularFeature<IPropertyAccessEditor>("PropertyAccessEditor");
