PropertyBindings: With function binding, display the error icon like a regular binding.

diff --git a/Engine/Plugins/Runtime/PropertyBindingUtils/Source/PropertyBindingUtilsEditor/Private/PropertyBindingExtension.cpp b/Engine/Plugins/Runtime/PropertyBindingUtils/Source/PropertyBindingUtilsEditor/Private/PropertyBindingExtension.cpp
--- a/Engine/Plugins/Runtime/PropertyBindingUtils/Source/PropertyBindingUtilsEditor/Private/PropertyBindingExtension.cpp
+++ b/Engine/Plugins/Runtime/PropertyBindingUtils/Source/PropertyBindingUtilsEditor/Private/PropertyBindingExtension.cpp
@@ -492,6 +492,7 @@ void FCachedBindingData::UpdateData()
 	FormatableTooltipText = FText::GetEmpty();
 	Color = FLinearColor::White;
 	Image = nullptr;
+	bIsCachedDataValid = false;
 
 	if (!PropertyHandle.IsValid())
 	{
@@ -623,6 +624,8 @@ void FCachedBindingData::UpdateData()
 
 				Image = FAppStyle::GetBrush(IconName);
 				Color = Schema->GetPinTypeColor(PinType);
+
+				bIsCachedDataValid = true;
 			}
 			else
 			{
@@ -679,6 +682,8 @@ void FCachedBindingData::UpdateData()
 		Image = FAppStyle::GetBrush(IconName);
 		Color = Schema->GetPinTypeColor(PinType);
 
+		bIsCachedDataValid = true;
+
 		CachedSourcePath.Reset();
 	}
 
@@ -1012,14 +1017,17 @@ FText FCachedBindingData::GetText()
 {
 	ConditionallyUpdateData();
 
-	// Bound PropertyFunction is allowed to override its display name.
-	FText CustomText;
-	if (ExecuteOnFunctionStructView(WeakBindingsOwner.Get(), TargetPath, [&CustomText, this](FConstStructView StructView)
-		{
-			return GetPropertyFunctionText(StructView, CustomText);
-		}))
+	if (bIsCachedDataValid)
 	{
-		return CustomText;
+		// Bound PropertyFunction is allowed to override its display name.
+		FText CustomText;
+		if (ExecuteOnFunctionStructView(WeakBindingsOwner.Get(), TargetPath, [&CustomText, this](FConstStructView StructView)
+			{
+				return GetPropertyFunctionText(StructView, CustomText);
+			}))
+		{
+			return CustomText;
+		}
 	}
 
 	return FText::FormatNamed(FormatableText, TEXT("SourceStruct"), SourceStructName);
@@ -1029,14 +1037,17 @@ FText FCachedBindingData::GetTooltipText()
 {
 	ConditionallyUpdateData();
 
-	// If the source property is a PropertyFunction and it overrides its display name, it's been used in the tooltip text.
-	FText CustomText;
-	if (ExecuteOnFunctionStructView(WeakBindingsOwner.Get(), TargetPath, [&CustomText, this](FConstStructView StructView)
-		{
-			return GetPropertyFunctionTooltipText(StructView, CustomText);
-		}))
+	if (bIsCachedDataValid)
 	{
-		return CustomText;
+		// If the source property is a PropertyFunction and it overrides its display name, it's been used in the tooltip text.
+		FText CustomText;
+		if (ExecuteOnFunctionStructView(WeakBindingsOwner.Get(), TargetPath, [&CustomText, this](FConstStructView StructView)
+			{
+				return GetPropertyFunctionTooltipText(StructView, CustomText);
+			}))
+		{
+			return CustomText;
+		}
 	}
 
 	return FText::FormatNamed(FormatableTooltipText, TEXT("SourceStruct"), SourceStructName);
@@ -1046,18 +1057,22 @@ FLinearColor FCachedBindingData::GetColor()
 {
 	ConditionallyUpdateData();
 
-	// Bound PropertyFunction is allowed to override its icon color if the binding leads directly into it's single output property.
-	if (CachedSourcePath.NumSegments() == 1)
+	if (bIsCachedDataValid)
 	{
-		FLinearColor CustomColor;
-		if (ExecuteOnFunctionStructView(WeakBindingsOwner.Get(), TargetPath, [&CustomColor, this](FConstStructView StructView)
-			{
-				return GetPropertyFunctionIconColor(StructView, CustomColor);
-			}))
+		// Bound PropertyFunction is allowed to override its icon color if the binding leads directly into it's single output property.
+		if (CachedSourcePath.NumSegments() == 1)
 		{
-			return CustomColor;
+			FLinearColor CustomColor;
+			if (ExecuteOnFunctionStructView(WeakBindingsOwner.Get(), TargetPath, [&CustomColor, this](FConstStructView StructView)
+				{
+					return GetPropertyFunctionIconColor(StructView, CustomColor);
+				}))
+			{
+				return CustomColor;
+			}
 		}
 	}
+
 	return Color;
 }
 
@@ -1065,14 +1080,17 @@ const FSlateBrush* FCachedBindingData::GetImage()
 {
 	ConditionallyUpdateData();
 
-	// Bound PropertyFunction is allowed to override its icon.
-	const FSlateBrush* CustomImage;
-	if (ExecuteOnFunctionStructView(WeakBindingsOwner.Get(), TargetPath, [&CustomImage, this](FConstStructView StructView)
-		{
-			return GetPropertyFunctionImage(StructView, CustomImage);
-		}))
+	if (bIsCachedDataValid)
 	{
-		return CustomImage;
+		// Bound PropertyFunction is allowed to override its icon.
+		const FSlateBrush* CustomImage;
+		if (ExecuteOnFunctionStructView(WeakBindingsOwner.Get(), TargetPath, [&CustomImage, this](FConstStructView StructView)
+			{
+				return GetPropertyFunctionImage(StructView, CustomImage);
+			}))
+		{
+			return CustomImage;
+		}
 	}
 
 	return Image;
diff --git a/Engine/Plugins/Runtime/PropertyBindingUtils/Source/PropertyBindingUtilsEditor/Public/PropertyBindingExtension.h b/Engine/Plugins/Runtime/PropertyBindingUtils/Source/PropertyBindingUtilsEditor/Public/PropertyBindingExtension.h
--- a/Engine/Plugins/Runtime/PropertyBindingUtils/Source/PropertyBindingUtilsEditor/Public/PropertyBindingExtension.h
+++ b/Engine/Plugins/Runtime/PropertyBindingUtils/Source/PropertyBindingUtilsEditor/Public/PropertyBindingExtension.h
@@ -157,6 +157,9 @@ private:
 	
 	FLinearColor Color = FLinearColor::White;
 	const FSlateBrush* Image = nullptr;
+
+	/** The binding is invalid. The image/color/tooltip/... represents an error state. */
+	bool bIsCachedDataValid = false;
 	
 	bool bIsDataCached = false;
 };
