[PropertyBindings] no longer prevent parameter creation if property is already bound to the target struct. It will simply create a new parameter with an incremental name

diff --git a/Engine/Plugins/Runtime/PropertyBindingUtils/Source/PropertyBindingUtilsEditor/Private/PropertyBindingExtension.cpp b/Engine/Plugins/Runtime/PropertyBindingUtils/Source/PropertyBindingUtilsEditor/Private/PropertyBindingExtension.cpp
--- a/Engine/Plugins/Runtime/PropertyBindingUtils/Source/PropertyBindingUtilsEditor/Private/PropertyBindingExtension.cpp
+++ b/Engine/Plugins/Runtime/PropertyBindingUtils/Source/PropertyBindingUtilsEditor/Private/PropertyBindingExtension.cpp
@@ -895,13 +895,6 @@ bool FCachedBindingData::CanCreateParameter(const FPropertyBindingBindableStruct
 		return false;
 	}
 
-	// Do not allow parameter creation if property is already bound to the struct
-	const FPropertyBindingPath* SourcePath = EditorBindings->GetBindingSource(TargetPath);
-	if (SourcePath && SourcePath->GetStructID() == InStructDesc.ID)
-	{
-		return false;
-	}
-
 	// Allow classes implementing IPropertyBindingBindingCollectionOwner to block parameter creation for that struct
 	if (!BindingsOwner->CanCreateParameter(InStructDesc.ID))
 	{
