StateTree: Add a binding viewer for the editor. That helps identify the bindings used by the state tree asset.

diff --git a/Engine/Plugins/Runtime/PropertyBindingUtils/Source/PropertyBindingUtilsEditor/Private/Widgets/SPropertyBindingView.cpp b/Engine/Plugins/Runtime/PropertyBindingUtils/Source/PropertyBindingUtilsEditor/Private/Widgets/SPropertyBindingView.cpp
new file mode 100644
--- /dev/null
+++ b/Engine/Plugins/Runtime/PropertyBindingUtils/Source/PropertyBindingUtilsEditor/Private/Widgets/SPropertyBindingView.cpp
@@ -0,0 +1,184 @@
+// Copyright Epic Games, Inc. All Rights Reserved.
+
+#include "Widgets/SPropertyBindingView.h"
+#include "PropertyBindingBinding.h"
+#include "PropertyBindingBindingCollection.h"
+
+#define LOCTEXT_NAMESPACE "SPropertyBindingViewer"
+
+namespace UE::PropertyBinding
+{
+namespace Private
+{
+	static const FLazyName ColumnId_SourceStruct = "SourceStruct";
+	static const FLazyName ColumnId_SourcePath = "SourcePath";
+	static const FLazyName ColumnId_TargetStruct = "TargetStruct";
+	static const FLazyName ColumnId_TargetPath = "TargetPath";
+
+	class SBindingViewRow : public SMultiColumnTableRow<TSharedPtr<SBindingView::FItem>>
+	{
+	public:
+		SLATE_BEGIN_ARGS(SBindingViewRow) {}
+		SLATE_END_ARGS()
+
+		void Construct(const FArguments& Args, TSharedPtr<SBindingView::FItem> InItem, TWeakInterfacePtr<const IPropertyBindingBindingCollectionOwner> InCollectionOwner, const TSharedRef<STableViewBase>& OwnerTableView)
+		{
+			Item = InItem;
+			CollectionOwner = InCollectionOwner;
+			SMultiColumnTableRow<TSharedPtr<SBindingView::FItem>>::Construct(
+				FSuperRowType::FArguments()
+				.Padding(1.0f),
+				OwnerTableView
+			);
+		}
+
+		FText HandleGetStructName(FGuid StructID) const
+		{
+			if (const IPropertyBindingBindingCollectionOwner* CollectionOwnerPtr = CollectionOwner.Get())
+			{
+				TInstancedStruct<FPropertyBindingBindableStructDescriptor> StructDesc;
+				if (CollectionOwnerPtr->GetBindableStructByID(StructID, StructDesc))
+				{
+					return FText::FromString(StructDesc.Get().ToString());
+				}
+			}
+			return FText::GetEmpty();
+		}
+
+		virtual TSharedRef<SWidget> GenerateWidgetForColumn(const FName& ColumnName) override
+		{
+			if (ColumnName == ColumnId_SourceStruct)
+			{
+				const FGuid StructID = Item->SourcePath.GetStructID();
+				return SNew(STextBlock)
+					.Text(this, &SBindingViewRow::HandleGetStructName, StructID);
+			}
+			else if (ColumnName == ColumnId_SourcePath)
+			{
+				return SNew(STextBlock)
+					.Text(FText::FromString(Item->SourcePath.ToString()));		
+			}
+			else if (ColumnName == ColumnId_TargetStruct)
+			{
+				const FGuid StructID = Item->TargetPath.GetStructID();
+				return SNew(STextBlock)
+					.Text(this, &SBindingViewRow::HandleGetStructName, StructID);
+			}
+			else if (ColumnName == ColumnId_TargetPath)
+			{
+				return SNew(STextBlock)
+					.Text(FText::FromString(Item->TargetPath.ToString()));
+			}
+
+			return SNullWidget::NullWidget;
+		}
+
+	private:
+		TSharedPtr<SBindingView::FItem> Item;
+		TWeakInterfacePtr<const IPropertyBindingBindingCollectionOwner> CollectionOwner;
+	};
+}
+
+void SBindingView::Construct(const FArguments& InArgs)
+{
+	OnGetBindingCollection = InArgs._GetBindingCollection;
+	CollectionOwner = InArgs._CollectionOwner;
+
+	ListView = SNew(SListView<TSharedPtr<FItem>>)
+		.SelectionMode(ESelectionMode::None)
+		.ListItemsSource(&Values)
+		.OnGenerateRow(this, &SBindingView::HandleGenerateRow)
+		.HeaderRow
+		(
+			SNew(SHeaderRow)
+			+ SHeaderRow::Column(Private::ColumnId_SourceStruct)
+			.DefaultLabel(LOCTEXT("SourceStructColumnLabel", "Source"))
+			
+			+ SHeaderRow::Column(Private::ColumnId_SourcePath)
+			.DefaultLabel(LOCTEXT("SourcePathColumnLabel", "Path"))
+
+			+ SHeaderRow::Column(Private::ColumnId_TargetStruct)
+			.DefaultLabel(LOCTEXT("TargetStructColumnLabel", "Target"))
+
+			+ SHeaderRow::Column(Private::ColumnId_TargetPath)
+			.DefaultLabel(LOCTEXT("TargetPathColumnLabel", "Path"))
+		);
+
+	ChildSlot
+	[
+		ListView.ToSharedRef()
+	];
+}
+
+const FPropertyBindingBindingCollection* SBindingView::GetBindingCollection()
+{
+	return OnGetBindingCollection.IsBound() ? OnGetBindingCollection.Execute() : nullptr;
+}
+
+void SBindingView::RequestRefresh()
+{
+	ListView->RequestListRefresh();
+}
+
+void SBindingView::Tick(const FGeometry& AllottedGeometry, const double InCurrentTime, const float InDeltaTime)
+{
+	SWidget::Tick(AllottedGeometry, InCurrentTime, InDeltaTime);
+
+	const FPropertyBindingBindingCollection* Collection = GetBindingCollection();
+	if (Collection)
+	{
+		int32 Index = 0;
+		bool bRequestRefresh = false;
+		Collection->ForEachBinding(
+			[&Index, &bRequestRefresh, Self=this](const FPropertyBindingBinding & Binding)
+			{
+				if (Self->Values.Num() <= Index)
+				{
+					Self->Values.Add(MakeShared<FItem>(Binding.GetSourcePath(), Binding.GetTargetPath(), Binding.GetPropertyFunctionNode().GetScriptStruct()));
+					bRequestRefresh = true;
+				}
+				else
+				{
+					const TSharedPtr<FItem>& Item = Self->Values[Index];
+					const bool bChanged = Item->FunctionNodeStruct.Get() != Binding.GetPropertyFunctionNode().GetScriptStruct()
+						|| Item->SourcePath != Binding.GetSourcePath()
+						|| Item->TargetPath != Binding.GetTargetPath();
+					if (bChanged)
+					{
+						Self->Values.SetNum(Index);
+						Self->Values.Add(MakeShared<FItem>(Binding.GetSourcePath(), Binding.GetTargetPath(), Binding.GetPropertyFunctionNode().GetScriptStruct()));
+						bRequestRefresh = true;
+					}
+				}
+				++Index;
+			});
+
+		if (Values.Num() != Index)
+		{
+			Values.SetNum(Index);
+			bRequestRefresh = true;
+		}
+
+		if (bRequestRefresh)
+		{
+			RequestRefresh();
+		}
+	}
+	else
+	{
+		if (!Values.IsEmpty())
+		{
+			Values.Empty();
+			RequestRefresh();
+		}
+	}
+}
+
+TSharedRef<ITableRow> SBindingView::HandleGenerateRow(TSharedPtr<FItem> Value, const TSharedRef<STableViewBase>& OwnerTable)
+{
+	return SNew(Private::SBindingViewRow, Value, CollectionOwner, OwnerTable);
+}
+
+} // namespace UE::PropertyBinding
+
+#undef LOCTEXT_NAMESPACE
diff --git a/Engine/Plugins/Runtime/PropertyBindingUtils/Source/PropertyBindingUtilsEditor/PropertyBindingUtilsEditor.Build.cs b/Engine/Plugins/Runtime/PropertyBindingUtils/Source/PropertyBindingUtilsEditor/PropertyBindingUtilsEditor.Build.cs
--- a/Engine/Plugins/Runtime/PropertyBindingUtils/Source/PropertyBindingUtilsEditor/PropertyBindingUtilsEditor.Build.cs
+++ b/Engine/Plugins/Runtime/PropertyBindingUtils/Source/PropertyBindingUtilsEditor/PropertyBindingUtilsEditor.Build.cs
@@ -15,6 +15,7 @@ namespace UnrealBuildTool.Rules
 					"CoreUObject",
 					"EditorFramework",
 					"Engine",
+					"InputCore",
 					"PropertyBindingUtils",
 					"Slate",
 					"SlateCore",
diff --git a/Engine/Plugins/Runtime/PropertyBindingUtils/Source/PropertyBindingUtilsEditor/Public/Widgets/SPropertyBindingView.h b/Engine/Plugins/Runtime/PropertyBindingUtils/Source/PropertyBindingUtilsEditor/Public/Widgets/SPropertyBindingView.h
new file mode 100644
--- /dev/null
+++ b/Engine/Plugins/Runtime/PropertyBindingUtils/Source/PropertyBindingUtilsEditor/Public/Widgets/SPropertyBindingView.h
@@ -0,0 +1,56 @@
+// Copyright Epic Games, Inc. All Rights Reserved.
+
+#pragma once
+
+#include "PropertyBindingBinding.h"
+#include "PropertyBindingBindingCollectionOwner.h"
+#include "UObject/WeakInterfacePtr.h"
+
+#include "Widgets/SCompoundWidget.h"
+#include "Widgets/Views/SListView.h"
+
+struct FPropertyBindingBindingCollection;
+
+namespace UE::PropertyBinding
+{
+	namespace Private
+	{
+		class SBindingViewRow;
+	}
+
+	class SBindingView : public SCompoundWidget
+	{
+	private:
+		friend ::UE::PropertyBinding::Private::SBindingViewRow;
+		struct FItem
+		{
+			FPropertyBindingPath SourcePath;
+			FPropertyBindingPath TargetPath;
+			TWeakObjectPtr<const UScriptStruct> FunctionNodeStruct;
+		};
+
+	public:
+		DECLARE_DELEGATE_RetVal(const FPropertyBindingBindingCollection*, FOnGetBindingCollection);
+
+		SLATE_BEGIN_ARGS(SBindingView){}
+			SLATE_EVENT(FOnGetBindingCollection, GetBindingCollection)
+			SLATE_ARGUMENT(TScriptInterface<const IPropertyBindingBindingCollectionOwner>, CollectionOwner)
+		SLATE_END_ARGS()
+
+		PROPERTYBINDINGUTILSEDITOR_API void Construct(const FArguments& InArgs);
+		PROPERTYBINDINGUTILSEDITOR_API virtual void Tick(const FGeometry& AllottedGeometry, const double InCurrentTime, const float InDeltaTime) override;
+
+	private:
+		const FPropertyBindingBindingCollection* GetBindingCollection();
+		void RequestRefresh();
+
+		TSharedRef<ITableRow> HandleGenerateRow(TSharedPtr<FItem> Value, const TSharedRef<STableViewBase>& OwnerTable);
+
+	private:
+		FOnGetBindingCollection OnGetBindingCollection;
+		TWeakInterfacePtr<const IPropertyBindingBindingCollectionOwner> CollectionOwner;
+		TSharedPtr<SListView<TSharedPtr<FItem>>> ListView;
+		TArray<TSharedPtr<FItem>> Values;
+};
+
+} // namespace UE::PropertyBinding
\ No newline at end of file
diff --git a/Engine/Plugins/Runtime/StateTree/Source/StateTreeEditorModule/Private/StateTreeEditor.cpp b/Engine/Plugins/Runtime/StateTree/Source/StateTreeEditorModule/Private/StateTreeEditor.cpp
--- a/Engine/Plugins/Runtime/StateTree/Source/StateTreeEditorModule/Private/StateTreeEditor.cpp
+++ b/Engine/Plugins/Runtime/StateTree/Source/StateTreeEditorModule/Private/StateTreeEditor.cpp
@@ -207,6 +207,7 @@ void FStateTreeEditor::InitEditor( const EToolkitMode::Type Mode, const TSharedP
 					->AddTab(CompilerResultsTabId, ETabState::ClosedTab)
 					->AddTab(UE::StateTreeEditor::FWorkspaceTabHost::SearchTabId.Resolve(), ETabState::ClosedTab)
 					->AddTab(UE::StateTreeEditor::FWorkspaceTabHost::DebuggerTabId.Resolve(), ETabState::ClosedTab)
+					->AddTab(UE::StateTreeEditor::FWorkspaceTabHost::BindingTabId.Resolve(), ETabState::ClosedTab)
 				)
 			)
 			->Split
diff --git a/Engine/Plugins/Runtime/StateTree/Source/StateTreeEditorModule/Private/StateTreeEditorLayoutHost.cpp b/Engine/Plugins/Runtime/StateTree/Source/StateTreeEditorModule/Private/StateTreeEditorLayoutHost.cpp
--- a/Engine/Plugins/Runtime/StateTree/Source/StateTreeEditorModule/Private/StateTreeEditorLayoutHost.cpp
+++ b/Engine/Plugins/Runtime/StateTree/Source/StateTreeEditorModule/Private/StateTreeEditorLayoutHost.cpp
@@ -11,21 +11,23 @@ class SDockTab;
 
 namespace UE::StateTreeEditor
 {
+const FLazyName FWorkspaceTabHost::BindingTabId = "StateTreeEditor_Binding";
+const FLazyName FWorkspaceTabHost::DebuggerTabId = "StateTreeEditor_Debugger";
 const FLazyName FWorkspaceTabHost::OutlinerTabId = "StateTreeEditor_StateTreeOutliner";
 const FLazyName FWorkspaceTabHost::SearchTabId = "StateTreeEditor_StateTreeSearch";
 const FLazyName FWorkspaceTabHost::StatisticsTabId = "StateTreeEditor_StateTreeStatistics";
-const FLazyName FWorkspaceTabHost::DebuggerTabId = "StateTreeEditor_Debugger";
 
 TConstArrayView<FMinorWorkspaceTabConfig> FWorkspaceTabHost::GetTabConfigs() const
 {
 	auto BuildTabConfigs = []()
 		{
 			TArray<FMinorWorkspaceTabConfig> TabConfigs;
+			TabConfigs.Emplace(BindingTabId.Resolve(), LOCTEXT("StateTreeBindingTab", "Bindings"), FText(), FSlateIcon(FAppStyle::GetAppStyleSetName(), "Icons.Link"), UAssetEditorUISubsystem::BottomRightTabID);
 			TabConfigs.Emplace(OutlinerTabId.Resolve(), LOCTEXT("StateTreeOutlinerTab", "Outliner"), FText(), FSlateIcon(FAppStyle::GetAppStyleSetName(), "LevelEditor.Tabs.Outliner"), UAssetEditorUISubsystem::TopLeftTabID);
 			TabConfigs.Emplace(SearchTabId.Resolve(), LOCTEXT("StateTreeSearchTab", "Find"), FText(), FSlateIcon(FAppStyle::GetAppStyleSetName(), "BlueprintEditor.FindInBlueprints.MenuIcon"), UAssetEditorUISubsystem::BottomLeftTabID);
 			TabConfigs.Emplace(StatisticsTabId.Resolve(), LOCTEXT("StatisticsTitle", "Statistics"), FText(), FSlateIcon(FAppStyle::GetAppStyleSetName(), "AssetEditor.ToggleStats"), UAssetEditorUISubsystem::BottomRightTabID);
 #if WITH_STATETREE_TRACE_DEBUGGER
-			TabConfigs.Emplace(DebuggerTabId.Resolve(), LOCTEXT("DebuggerTab", "Debugger"), FText(), FSlateIcon(FAppStyle::GetAppStyleSetName(), "Debug"), UAssetEditorUISubsystem::TopRightTabID);
+			TabConfigs.Emplace(DebuggerTabId.Resolve(), LOCTEXT("DebuggerTab", "Debugger"), FText(), FSlateIcon(FAppStyle::GetAppStyleSetName(), "Debug"), UAssetEditorUISubsystem::BottomRightTabID);
 #endif
 			return TabConfigs;
 		};
diff --git a/Engine/Plugins/Runtime/StateTree/Source/StateTreeEditorModule/Private/StateTreeEditorModeToolkit.cpp b/Engine/Plugins/Runtime/StateTree/Source/StateTreeEditorModule/Private/StateTreeEditorModeToolkit.cpp
--- a/Engine/Plugins/Runtime/StateTree/Source/StateTreeEditorModule/Private/StateTreeEditorModeToolkit.cpp
+++ b/Engine/Plugins/Runtime/StateTree/Source/StateTreeEditorModule/Private/StateTreeEditorModeToolkit.cpp
@@ -26,6 +26,7 @@
 #include "Toolkits/AssetEditorModeUILayer.h"
 #include "Widgets/Input/SMultiLineEditableTextBox.h"
 #include "Widgets/Docking/SDockTab.h"
+#include "Widgets/SPropertyBindingView.h"
 #include "ToolMenus.h"
 #include "Subsystems/AssetEditorSubsystem.h"
 #include "Framework/Application/SlateApplication.h"
@@ -108,7 +109,30 @@ void FStateTreeEditorModeToolkit::InvokeUI()
 
 void FStateTreeEditorModeToolkit::HandleTabSpawned(UE::StateTreeEditor::FSpawnedWorkspaceTab SpawnedTab)
 {
-	if (SpawnedTab.TabID == UE::StateTreeEditor::FWorkspaceTabHost::OutlinerTabId)
+	if (SpawnedTab.TabID == UE::StateTreeEditor::FWorkspaceTabHost::BindingTabId)
+	{
+		if (TSharedPtr<SDockTab> DockTab = SpawnedTab.DockTab.Pin())
+		{
+			const UStateTreeEditorData* EditorData = nullptr;
+			if (UStateTreeEditorMode* EditorMode = WeakEditorMode.Get())
+			{
+				if (UStateTree* StateTree = EditorMode->GetStateTree())
+				{
+					if (UStateTreeEditingSubsystem* Subsystem = GEditor->GetEditorSubsystem<UStateTreeEditingSubsystem>())
+					{
+						EditorData = Subsystem->FindOrAddViewModel(StateTree)->GetStateTreeEditorData();
+					}
+				}
+			}
+
+			DockTab->SetContent(
+				SNew(UE::PropertyBinding::SBindingView)
+				.GetBindingCollection(this, &FStateTreeEditorModeToolkit::GetBindingCollection)
+				.CollectionOwner(TScriptInterface<const IPropertyBindingBindingCollectionOwner>(EditorData))
+			);
+		}
+	}
+	else if (SpawnedTab.TabID == UE::StateTreeEditor::FWorkspaceTabHost::OutlinerTabId)
 	{
 		WeakOutlinerTab = SpawnedTab.DockTab;
 		UpdateStateTreeOutliner();
@@ -149,15 +173,10 @@ void FStateTreeEditorModeToolkit::HandleTabSpawned(UE::StateTreeEditor::FSpawned
 
 void FStateTreeEditorModeToolkit::HandleTabClosed(UE::StateTreeEditor::FSpawnedWorkspaceTab SpawnedTab)
 {
-	if (SpawnedTab.TabID == UE::StateTreeEditor::FWorkspaceTabHost::DebuggerTabId)
+	if (TSharedPtr<SDockTab> DockTab = WeakDebuggerTab.Pin())
 	{
-#if WITH_STATETREE_TRACE_DEBUGGER
-		if (TSharedPtr<SDockTab> DockTab = WeakDebuggerTab.Pin())
-		{
-			// Destroy the SStateTreeDebuggerView
-			DockTab->SetContent(SNullWidget::NullWidget);
-		}
-#endif // WITH_STATETREE_TRACE_DEBUGGER
+		// Destroy the inner widget.
+		DockTab->SetContent(SNullWidget::NullWidget);
 	}
 }
 
@@ -580,6 +599,24 @@ FText FStateTreeEditorModeToolkit::GetStatisticsText() const
 	return FText::Join(FText::FromString(TEXT("\n")), Rows);
 }
 
+const FPropertyBindingBindingCollection* FStateTreeEditorModeToolkit::GetBindingCollection() const
+{
+	if (TStrongObjectPtr<UStateTreeEditorMode> EditorMode = WeakEditorMode.Pin())
+	{
+		if (UStateTree* StateTree = EditorMode->GetStateTree())
+		{
+			if (UStateTreeEditingSubsystem* Subsystem = GEditor->GetEditorSubsystem<UStateTreeEditingSubsystem>())
+			{
+				if (const UStateTreeEditorData* EditorData = Subsystem->FindOrAddViewModel(StateTree)->GetStateTreeEditorData())
+				{
+					return EditorData->GetPropertyEditorBindings();
+				}
+			}
+		}
+	}
+	return nullptr;
+}
+
 void FStateTreeEditorModeToolkit::UpdateStateTreeOutliner()
 {
 	StateTreeOutliner = SNullWidget::NullWidget;
diff --git a/Engine/Plugins/Runtime/StateTree/Source/StateTreeEditorModule/Public/StateTreeEditorModeToolkit.h b/Engine/Plugins/Runtime/StateTree/Source/StateTreeEditorModule/Public/StateTreeEditorModeToolkit.h
--- a/Engine/Plugins/Runtime/StateTree/Source/StateTreeEditorModule/Public/StateTreeEditorModeToolkit.h
+++ b/Engine/Plugins/Runtime/StateTree/Source/StateTreeEditorModule/Public/StateTreeEditorModeToolkit.h
@@ -7,6 +7,7 @@
 #include "Toolkits/BaseToolkit.h"
 
 class UStateTreeEditorMode;
+struct FPropertyBindingBindingCollection;
 namespace UE::StateTreeEditor
 {
 	struct FSpawnedWorkspaceTab;
@@ -45,6 +46,7 @@ protected:
 	void OnNodeBPBaseClassPicked(UClass* NodeClass) const;
 	
 	FText GetStatisticsText() const;
+	const FPropertyBindingBindingCollection* GetBindingCollection() const;
 
 	void UpdateStateTreeOutliner();
 
diff --git a/Engine/Plugins/Runtime/StateTree/Source/StateTreeEditorModule/Public/StateTreeEditorWorkspaceTabHost.h b/Engine/Plugins/Runtime/StateTree/Source/StateTreeEditorModule/Public/StateTreeEditorWorkspaceTabHost.h
--- a/Engine/Plugins/Runtime/StateTree/Source/StateTreeEditorModule/Public/StateTreeEditorWorkspaceTabHost.h
+++ b/Engine/Plugins/Runtime/StateTree/Source/StateTreeEditorModule/Public/StateTreeEditorWorkspaceTabHost.h
@@ -29,10 +29,11 @@ namespace UE::StateTreeEditor
 	class STATETREEEDITORMODULE_API FWorkspaceTabHost : public TSharedFromThis<FWorkspaceTabHost>
 	{
 	public:
+		static const FLazyName BindingTabId;
+		static const FLazyName DebuggerTabId;
 		static const FLazyName OutlinerTabId;
 		static const FLazyName SearchTabId;
 		static const FLazyName StatisticsTabId;
-		static const FLazyName DebuggerTabId;
 
 	public:
 		virtual ~FWorkspaceTabHost() = default;
