[PropertyBindings] reverting change that was preventing EditConst properties (e.g., VisibleAnywhere) to be used in bindings. Already too many uses to change that behavior.

diff --git a/Engine/Plugins/Runtime/PropertyBindingUtils/Source/PropertyBindingUtilsEditor/Private/PropertyBindingExtension.cpp b/Engine/Plugins/Runtime/PropertyBindingUtils/Source/PropertyBindingUtilsEditor/Private/PropertyBindingExtension.cpp
--- a/Engine/Plugins/Runtime/PropertyBindingUtils/Source/PropertyBindingUtilsEditor/Private/PropertyBindingExtension.cpp
+++ b/Engine/Plugins/Runtime/PropertyBindingUtils/Source/PropertyBindingUtilsEditor/Private/PropertyBindingExtension.cpp
@@ -1130,13 +1130,6 @@ bool IsPropertyBindable(const FProperty& Property)
 		return false;
 	}
 
-	if (const bool bIsEditConst = Property.HasAnyPropertyFlags(CPF_EditConst))
-	{
-		UE_LOG(LogPropertyBindingUtils, Verbose, TEXT("Property %s is not bindable because it's readonly in the editor"),
-			*Property.GetName());
-		return false;
-	}
-
 	const bool bPrivateOrProtected = !Property.HasAnyPropertyFlags(CPF_NativeAccessSpecifierPrivate | CPF_NativeAccessSpecifierProtected);
 	const bool bPrivateButBlueprintAccessible = Property.GetBoolMetaData(FBlueprintMetadata::MD_AllowPrivateAccess);
 	if (!bPrivateOrProtected && !bPrivateButBlueprintAccessible)
