Add FPropertyBindingPropertyIndirection GetTypeHash

diff --git a/Engine/Plugins/Runtime/PropertyBindingUtils/Source/PropertyBindingUtils/Public/PropertyBindingPath.h b/Engine/Plugins/Runtime/PropertyBindingUtils/Source/PropertyBindingUtils/Public/PropertyBindingPath.h
--- a/Engine/Plugins/Runtime/PropertyBindingUtils/Source/PropertyBindingUtils/Public/PropertyBindingPath.h
+++ b/Engine/Plugins/Runtime/PropertyBindingUtils/Source/PropertyBindingUtils/Public/PropertyBindingPath.h
@@ -502,6 +502,17 @@ struct FPropertyBindingPropertyIndirection
 
 	/** Cached array property. */
 	const FArrayProperty* ArrayProperty = nullptr;
+
+	friend uint32 GetTypeHash(const FPropertyBindingPropertyIndirection& Indirection)
+	{
+		uint32 Hash = GetTypeHash(Indirection.ArrayIndex);
+		Hash = HashCombineFast(Hash, GetTypeHash(Indirection.Offset));
+		Hash = HashCombineFast(Hash, GetTypeHash(Indirection.NextIndex));
+		Hash = HashCombineFast(Hash, GetTypeHash(Indirection.Type));
+		Hash = HashCombineFast(Hash, GetTypeHash(Indirection.InstanceStruct));
+		
+		return Hash;
+	}
 };
 
 /**
