[PropertyBindings] fixed issue when reinstancing objects in source struct descriptors
Also converted casts to CastFieldChecked for property based variables

diff --git a/Engine/Plugins/Runtime/PropertyBindingUtils/Source/PropertyBindingUtils/Private/PropertyBindingBindingCollection.cpp b/Engine/Plugins/Runtime/PropertyBindingUtils/Source/PropertyBindingUtils/Private/PropertyBindingBindingCollection.cpp
--- a/Engine/Plugins/Runtime/PropertyBindingUtils/Source/PropertyBindingUtils/Private/PropertyBindingBindingCollection.cpp
+++ b/Engine/Plugins/Runtime/PropertyBindingUtils/Source/PropertyBindingUtils/Private/PropertyBindingBindingCollection.cpp
@@ -927,31 +927,31 @@ bool FPropertyBindingBindingCollection::PerformCopy(const FPropertyBindingCopyIn
 		Copy.TargetLeafProperty->CopyCompleteValue(TargetAddress, SourceAddress);
 		break;
 	case EPropertyCopyType::CopyBool:
-		static_cast<const FBoolProperty*>(Copy.TargetLeafProperty)->SetPropertyValue(TargetAddress, static_cast<const FBoolProperty*>(Copy.SourceLeafProperty)->GetPropertyValue(SourceAddress));
+		CastFieldChecked<const FBoolProperty>(Copy.TargetLeafProperty)->SetPropertyValue(TargetAddress, CastFieldChecked<const FBoolProperty>(Copy.SourceLeafProperty)->GetPropertyValue(SourceAddress));
 		break;
 	case EPropertyCopyType::CopyStruct:
 		// If SourceProperty == nullptr (pointing to the struct source directly), the GetAddress() did the right thing and is pointing to the beginning of the struct. 
-		static_cast<const FStructProperty*>(Copy.TargetLeafProperty)->Struct->CopyScriptStruct(TargetAddress, SourceAddress);
+		CastFieldChecked<const FStructProperty>(Copy.TargetLeafProperty)->Struct->CopyScriptStruct(TargetAddress, SourceAddress);
 		break;
 	case EPropertyCopyType::CopyObject:
 		if (Copy.SourceLeafProperty == nullptr || SourceAddress == nullptr)
 		{
 			// Source is pointing at object directly.
-			static_cast<const FObjectPropertyBase*>(Copy.TargetLeafProperty)->SetObjectPropertyValue(TargetAddress, (UObject*)SourceAddress);
+			CastFieldChecked<const FObjectPropertyBase>(Copy.TargetLeafProperty)->SetObjectPropertyValue(TargetAddress, reinterpret_cast<UObject*>(SourceAddress));
 		}
 		else
 		{
-			static_cast<const FObjectPropertyBase*>(Copy.TargetLeafProperty)->SetObjectPropertyValue(TargetAddress, static_cast<const FObjectPropertyBase*>(Copy.SourceLeafProperty)->GetObjectPropertyValue(SourceAddress));
+			CastFieldChecked<const FObjectPropertyBase>(Copy.TargetLeafProperty)->SetObjectPropertyValue(TargetAddress, CastFieldChecked<const FObjectPropertyBase>(Copy.SourceLeafProperty)->GetObjectPropertyValue(SourceAddress));
 		}
 		break;
 	case EPropertyCopyType::CopyName:
-		static_cast<const FNameProperty*>(Copy.TargetLeafProperty)->SetPropertyValue(TargetAddress, static_cast<const FNameProperty*>(Copy.SourceLeafProperty)->GetPropertyValue(SourceAddress));
+		CastFieldChecked<const FNameProperty>(Copy.TargetLeafProperty)->SetPropertyValue(TargetAddress, CastFieldChecked<const FNameProperty>(Copy.SourceLeafProperty)->GetPropertyValue(SourceAddress));
 		break;
 	case EPropertyCopyType::CopyFixedArray:
 	{
 		// Copy into fixed sized array (EditFixedSize). Resizable arrays are copied as Complex, and regular fixed sizes arrays via the regular copies (dim specifies array size).
-		const FArrayProperty* SourceArrayProperty = static_cast<const FArrayProperty*>(Copy.SourceLeafProperty);
-		const FArrayProperty* TargetArrayProperty = static_cast<const FArrayProperty*>(Copy.TargetLeafProperty);
+		const FArrayProperty* SourceArrayProperty = CastFieldChecked<const FArrayProperty>(Copy.SourceLeafProperty);
+		const FArrayProperty* TargetArrayProperty = CastFieldChecked<const FArrayProperty>(Copy.TargetLeafProperty);
 		FScriptArrayHelper SourceArrayHelper(SourceArrayProperty, SourceAddress);
 		FScriptArrayHelper TargetArrayHelper(TargetArrayProperty, TargetAddress);
 			
@@ -965,7 +965,7 @@ bool FPropertyBindingBindingCollection::PerformCopy(const FPropertyBindingCopyIn
 	case EPropertyCopyType::StructReference:
 	{
 		checkf(PropertyReferenceCopyFunc, TEXT("Not expecting EPropertyCopyType::StructReference if copy functor was not provided"));
-		const FStructProperty* SourceStructProperty = static_cast<const FStructProperty*>(Copy.SourceLeafProperty);
+		const FStructProperty* SourceStructProperty = CastFieldChecked<const FStructProperty>(Copy.SourceLeafProperty);
 		if (ensure(SourceStructProperty))
 		{
 			PropertyReferenceCopyFunc(*SourceStructProperty, SourceAddress, TargetAddress);
@@ -974,22 +974,22 @@ bool FPropertyBindingBindingCollection::PerformCopy(const FPropertyBindingCopyIn
 	}
 	// Bool promotions
 	case EPropertyCopyType::PromoteBoolToByte:
-		*reinterpret_cast<uint8*>(TargetAddress) = (uint8)static_cast<const FBoolProperty*>(Copy.SourceLeafProperty)->GetPropertyValue(SourceAddress);
+		*reinterpret_cast<uint8*>(TargetAddress) = (uint8)CastFieldChecked<const FBoolProperty>(Copy.SourceLeafProperty)->GetPropertyValue(SourceAddress);
 		break;
 	case EPropertyCopyType::PromoteBoolToInt32:
-		*reinterpret_cast<int32*>(TargetAddress) = (int32)static_cast<const FBoolProperty*>(Copy.SourceLeafProperty)->GetPropertyValue(SourceAddress);
+		*reinterpret_cast<int32*>(TargetAddress) = (int32)CastFieldChecked<const FBoolProperty>(Copy.SourceLeafProperty)->GetPropertyValue(SourceAddress);
 		break;
 	case EPropertyCopyType::PromoteBoolToUInt32:
-		*reinterpret_cast<uint32*>(TargetAddress) = (uint32)static_cast<const FBoolProperty*>(Copy.SourceLeafProperty)->GetPropertyValue(SourceAddress);
+		*reinterpret_cast<uint32*>(TargetAddress) = (uint32)CastFieldChecked<const FBoolProperty>(Copy.SourceLeafProperty)->GetPropertyValue(SourceAddress);
 		break;
 	case EPropertyCopyType::PromoteBoolToInt64:
-		*reinterpret_cast<int64*>(TargetAddress) = (int64)static_cast<const FBoolProperty*>(Copy.SourceLeafProperty)->GetPropertyValue(SourceAddress);
+		*reinterpret_cast<int64*>(TargetAddress) = (int64)CastFieldChecked<const FBoolProperty>(Copy.SourceLeafProperty)->GetPropertyValue(SourceAddress);
 		break;
 	case EPropertyCopyType::PromoteBoolToFloat:
-		*reinterpret_cast<float*>(TargetAddress) = (float)static_cast<const FBoolProperty*>(Copy.SourceLeafProperty)->GetPropertyValue(SourceAddress);
+		*reinterpret_cast<float*>(TargetAddress) = (float)CastFieldChecked<const FBoolProperty>(Copy.SourceLeafProperty)->GetPropertyValue(SourceAddress);
 		break;
 	case EPropertyCopyType::PromoteBoolToDouble:
-		*reinterpret_cast<double*>(TargetAddress) = (double)static_cast<const FBoolProperty*>(Copy.SourceLeafProperty)->GetPropertyValue(SourceAddress);
+		*reinterpret_cast<double*>(TargetAddress) = (double)CastFieldChecked<const FBoolProperty>(Copy.SourceLeafProperty)->GetPropertyValue(SourceAddress);
 		break;
 		
 	// Byte promotions	
@@ -1108,10 +1108,10 @@ void FPropertyBindingBindingCollection::PerformResetObjects(const FPropertyBindi
 		Copy.TargetLeafProperty->ClearValue(TargetAddress);
 		break;
 	case EPropertyCopyType::CopyStruct:
-		static_cast<const FStructProperty*>(Copy.TargetLeafProperty)->Struct->ClearScriptStruct(TargetAddress);
+		CastFieldChecked<const FStructProperty>(Copy.TargetLeafProperty)->Struct->ClearScriptStruct(TargetAddress);
 		break;
 	case EPropertyCopyType::CopyObject:
-		static_cast<const FObjectPropertyBase*>(Copy.TargetLeafProperty)->SetObjectPropertyValue(TargetAddress, nullptr);
+		CastFieldChecked<const FObjectPropertyBase>(Copy.TargetLeafProperty)->SetObjectPropertyValue(TargetAddress, nullptr);
 		break;
 	case EPropertyCopyType::StructReference:
 		checkf(PropertyReferenceResetFunc, TEXT("Not expecting EPropertyCopyType::StructReference if reset object functor was not provided"));
@@ -1122,7 +1122,7 @@ void FPropertyBindingBindingCollection::PerformResetObjects(const FPropertyBindi
 	case EPropertyCopyType::CopyFixedArray:
 	{
 		// Copy into fixed sized array (EditFixedSize). Resizable arrays are copied as Complex, and regular fixed sizes arrays via the regular copies (dim specifies array size).
-		const FArrayProperty* TargetArrayProperty = static_cast<const FArrayProperty*>(Copy.TargetLeafProperty);
+		const FArrayProperty* TargetArrayProperty = CastFieldChecked<const FArrayProperty>(Copy.TargetLeafProperty);
 		FScriptArrayHelper TargetArrayHelper(TargetArrayProperty, TargetAddress);
 		for (int32 ElementIndex = 0; ElementIndex < TargetArrayHelper.Num(); ++ElementIndex)
 		{
@@ -1177,15 +1177,22 @@ bool FPropertyBindingBindingCollection::ResetObjects(const FPropertyBindingIndex
 bool FPropertyBindingBindingCollection::ContainsAnyStruct(const TSet<const UStruct*>& InStructs) const
 {
 	// Look in derived classes source struct descriptors
-	VisitSourceStructDescriptorInternal([&InStructs](const FPropertyBindingBindableStructDescriptor& Descriptor)
+	bool bFoundInSourceStructs = false;
+	VisitSourceStructDescriptorInternal([&InStructs, &bFoundInSourceStructs](const FPropertyBindingBindableStructDescriptor& Descriptor)
 		{
 			if (InStructs.Contains(Descriptor.Struct))
 			{
+				bFoundInSourceStructs = true;
 				return EVisitResult::Break;
 			}
 			return EVisitResult::Continue;
 		});
 
+	if (bFoundInSourceStructs)
+	{
+		return true;
+	}
+
 	for (const FPropertyBindingCopyInfoBatch& CopyBatch : CopyBatches)
 	{
 		if (InStructs.Contains(CopyBatch.TargetStruct.Get().Struct))
