PropertyBinding: Add ability for promotions to keep the source property's metadata.

diff --git a/Engine/Plugins/Runtime/PropertyBindingUtils/Source/PropertyBindingUtils/Public/PropertyBindingBindingCollectionOwner.h b/Engine/Plugins/Runtime/PropertyBindingUtils/Source/PropertyBindingUtils/Public/PropertyBindingBindingCollectionOwner.h
--- a/Engine/Plugins/Runtime/PropertyBindingUtils/Source/PropertyBindingUtils/Public/PropertyBindingBindingCollectionOwner.h
+++ b/Engine/Plugins/Runtime/PropertyBindingUtils/Source/PropertyBindingUtils/Public/PropertyBindingBindingCollectionOwner.h
@@ -2,6 +2,7 @@
 
 #pragma once
 
+#include "StructUtils/PropertyBag.h"
 #include "UObject/Interface.h"
 #include "PropertyBindingBindingCollectionOwner.generated.h"
 
@@ -93,6 +94,20 @@ public:
 		return true;
 	}
 
+	/**
+	 * Called when a property is being promoted to a parameter.
+	 * @param PropertyDesc The newly created property descriptor
+	 */
+	virtual void OnPromotingToParameter(FPropertyBagPropertyDesc& PropertyDesc)
+	{
+		// By default, we create a descriptor based on the Target Property, but without the meta-data.
+		// This functionality mirrors the user action of adding a new property from the UI, where meta-data is not available.
+		// Additionally, meta-data like EditCondition is not desirable here.
+		// Derived classes can override this behavior.
+		PropertyDesc.MetaClass = nullptr;
+		PropertyDesc.MetaData.Reset();
+	}
+
 	/**
 	 * Creates the given properties in the property bag of the struct matching the given struct ID
 	 * @param InStructID Target struct ID
diff --git a/Engine/Plugins/Runtime/PropertyBindingUtils/Source/PropertyBindingUtilsEditor/Private/PropertyBindingExtension.cpp b/Engine/Plugins/Runtime/PropertyBindingUtils/Source/PropertyBindingUtilsEditor/Private/PropertyBindingExtension.cpp
--- a/Engine/Plugins/Runtime/PropertyBindingUtils/Source/PropertyBindingUtilsEditor/Private/PropertyBindingExtension.cpp
+++ b/Engine/Plugins/Runtime/PropertyBindingUtils/Source/PropertyBindingUtilsEditor/Private/PropertyBindingExtension.cpp
@@ -837,11 +837,7 @@ void FCachedBindingData::PromoteToParameter(const FName InPropertyName, TConstSt
 			PropertyCreationDesc.PropertyDesc = FPropertyBagPropertyDesc(InPropertyName, Property);
 		}
 
-		// Create desc based on the Target Property, but without the meta-data.
-		// This functionality mirrors the user action of adding a new property from the UI, where meta-data is not available.
-		// Additionally, meta-data like EditCondition is not desirable here
-		PropertyCreationDesc.PropertyDesc.MetaClass = nullptr;
-		PropertyCreationDesc.PropertyDesc.MetaData.Reset();
+		BindingsOwner->OnPromotingToParameter(PropertyCreationDesc.PropertyDesc);
 
 		// Set the Property & Container Address to copy
 		if (TargetProperty && TargetContainerAddress)
