PropertyBinding: Let extensions allow editor-only bindings.

diff --git a/Engine/Plugins/Runtime/PropertyBindingUtils/Source/PropertyBindingUtilsEditor/Private/PropertyBindingExtension.cpp b/Engine/Plugins/Runtime/PropertyBindingUtils/Source/PropertyBindingUtilsEditor/Private/PropertyBindingExtension.cpp
--- a/Engine/Plugins/Runtime/PropertyBindingUtils/Source/PropertyBindingUtilsEditor/Private/PropertyBindingExtension.cpp
+++ b/Engine/Plugins/Runtime/PropertyBindingUtils/Source/PropertyBindingUtilsEditor/Private/PropertyBindingExtension.cpp
@@ -1143,7 +1143,7 @@ bool IsPropertyBindable(const FProperty& Property)
 bool FPropertyBindingExtension::IsPropertyExtendable(const UClass* InObjectClass, const IPropertyHandle& InPropertyHandle) const
 {
 	const FProperty* Property = InPropertyHandle.GetProperty();
-	if (Property == nullptr || Property->HasAnyPropertyFlags(CPF_PersistentInstance | CPF_EditorOnly | CPF_Config | CPF_Deprecated))
+	if (Property == nullptr || Property->HasAnyPropertyFlags(GetDisallowedPropertyFlags()))
 	{
 		return false;
 	}
@@ -1462,6 +1462,11 @@ bool FPropertyBindingExtension::CanBindToProperty(const FPropertyBindingPath& In
 	return true;
 }
 
+uint64 FPropertyBindingExtension::GetDisallowedPropertyFlags() const
+{
+	return CPF_PersistentInstance | CPF_EditorOnly | CPF_Config | CPF_Deprecated;
+}
+
 bool FPropertyBindingExtension::CanPromoteToParameter(const TSharedPtr<IPropertyHandle>& InPropertyHandle) const
 {
 	const FProperty* Property = InPropertyHandle->GetProperty();
diff --git a/Engine/Plugins/Runtime/PropertyBindingUtils/Source/PropertyBindingUtilsEditor/Public/PropertyBindingExtension.h b/Engine/Plugins/Runtime/PropertyBindingUtils/Source/PropertyBindingUtilsEditor/Public/PropertyBindingExtension.h
--- a/Engine/Plugins/Runtime/PropertyBindingUtils/Source/PropertyBindingUtilsEditor/Public/PropertyBindingExtension.h
+++ b/Engine/Plugins/Runtime/PropertyBindingUtils/Source/PropertyBindingUtilsEditor/Public/PropertyBindingExtension.h
@@ -230,6 +230,7 @@ public:
 protected:
 	PROPERTYBINDINGUTILSEDITOR_API virtual TSharedPtr<UE::PropertyBinding::FCachedBindingData> CreateCachedBindingData(IPropertyBindingBindingCollectionOwner* InBindingsOwner, const FPropertyBindingPath& InTargetPath, const TSharedPtr<IPropertyHandle>& InPropertyHandle, TConstArrayView<TInstancedStruct<FPropertyBindingBindableStructDescriptor>> InAccessibleStructs) const;
 	PROPERTYBINDINGUTILSEDITOR_API virtual bool CanBindToProperty(const FPropertyBindingPath& InTargetPath, const IPropertyHandle& InPropertyHandle) const;
+	PROPERTYBINDINGUTILSEDITOR_API virtual uint64 GetDisallowedPropertyFlags() const;
 	PROPERTYBINDINGUTILSEDITOR_API virtual bool GetPromotionToParameterOverrideInternal(const FProperty& InProperty, bool& bOutOverride) const;
 	PROPERTYBINDINGUTILSEDITOR_API virtual void UpdateContextStruct(TConstStructView<FPropertyBindingBindableStructDescriptor> InStructDesc, FBindingContextStruct& InOutContextStruct, TMap<FString, FText>& InOutSectionNames) const;
 	virtual bool CanBindToArrayElements() const
