Expose some Property Binding Extension methods to make it easier to define custom menu actions in subclasses:
- Make it possible to specify binding icons from stylesets other than AppStyle via GetPinTypeAndIconForProperty
- Move UpdateData to protected so that subclasses can also define actions that changes the binding
- Add a CustomizeWidgetArgs method with which it's possible to populate and override anything in FPropertyBindingWidgetArgs

diff --git a/Engine/Plugins/Experimental/SceneState/Source/SceneStateBlueprintEditor/Private/DetailsView/SceneStateCachedBindingData.cpp b/Engine/Plugins/Experimental/SceneState/Source/SceneStateBlueprintEditor/Private/DetailsView/SceneStateCachedBindingData.cpp
--- a/Engine/Plugins/Experimental/SceneState/Source/SceneStateBlueprintEditor/Private/DetailsView/SceneStateCachedBindingData.cpp
+++ b/Engine/Plugins/Experimental/SceneState/Source/SceneStateBlueprintEditor/Private/DetailsView/SceneStateCachedBindingData.cpp
@@ -38,7 +38,7 @@ void FSceneStateCachedBindingData::UpdatePropertyReferenceTooltip(const FPropert
 	}
 }
 
-bool FSceneStateCachedBindingData::GetPinTypeAndIconForProperty(const FProperty& InProperty, FPropertyBindingDataView InTargetDataView, FEdGraphPinType& OutPinType, FName& OutIconName) const
+bool FSceneStateCachedBindingData::GetPinTypeAndIconForProperty(const FProperty& InProperty, FPropertyBindingDataView InTargetDataView, FEdGraphPinType& OutPinType, const FSlateBrush*& OutIconBrush) const
 {
 	if (UE::SceneState::IsPropertyReference(&InProperty) && InTargetDataView.IsValid())
 	{
@@ -48,7 +48,7 @@ bool FSceneStateCachedBindingData::GetPinTypeAndIconForProperty(const FProperty&
 		{
 			OutPinType = GetPropertyReferencePinType(&InProperty, TargetIndirections.Last().GetPropertyAddress());
 		}
-		OutIconName = TEXT("Kismet.Tabs.Variables");
+		OutIconBrush = FAppStyle::GetBrush("Kismet.Tabs.Variables");
 		return true;
 	}
 	return false;
diff --git a/Engine/Plugins/Experimental/SceneState/Source/SceneStateBlueprintEditor/Private/DetailsView/SceneStateCachedBindingData.h b/Engine/Plugins/Experimental/SceneState/Source/SceneStateBlueprintEditor/Private/DetailsView/SceneStateCachedBindingData.h
--- a/Engine/Plugins/Experimental/SceneState/Source/SceneStateBlueprintEditor/Private/DetailsView/SceneStateCachedBindingData.h
+++ b/Engine/Plugins/Experimental/SceneState/Source/SceneStateBlueprintEditor/Private/DetailsView/SceneStateCachedBindingData.h
@@ -14,7 +14,7 @@ struct FSceneStateCachedBindingData : UE::PropertyBinding::FCachedBindingData, T
 
 	//~ Begin UE::PropertyBinding::FCachedBindingData
 	virtual void UpdatePropertyReferenceTooltip(const FProperty& InProperty, FTextBuilder& InOutTextBuilder) const override;
-	virtual bool GetPinTypeAndIconForProperty(const FProperty& InProperty, FPropertyBindingDataView InTargetDataView, FEdGraphPinType& OutPinType, FName& OutIconName) const override;
+	virtual bool GetPinTypeAndIconForProperty(const FProperty& InProperty, FPropertyBindingDataView InTargetDataView, FEdGraphPinType& OutPinType, const FSlateBrush*& OutIconBrush) const override;
 	virtual bool IsPropertyReference(const FProperty& InProperty) override;
 	virtual void UpdateSourcePropertyPath(TConstStructView<FPropertyBindingBindableStructDescriptor> InDescriptor, const FPropertyBindingPath& InSourcePath, FString& OutString);
 	virtual void AddPropertyInfoOverride(const FProperty& InProperty, TArray<TSharedPtr<const UE::PropertyBinding::FPropertyInfoOverride>>& OutPropertyInfoOverrides) const override;
diff --git a/Engine/Plugins/Runtime/PropertyBindingUtils/Source/PropertyBindingUtilsEditor/Private/PropertyBindingExtension.cpp b/Engine/Plugins/Runtime/PropertyBindingUtils/Source/PropertyBindingUtilsEditor/Private/PropertyBindingExtension.cpp
--- a/Engine/Plugins/Runtime/PropertyBindingUtils/Source/PropertyBindingUtilsEditor/Private/PropertyBindingExtension.cpp
+++ b/Engine/Plugins/Runtime/PropertyBindingUtils/Source/PropertyBindingUtilsEditor/Private/PropertyBindingExtension.cpp
@@ -472,6 +472,12 @@ bool FCachedBindingData::GetPinTypeAndIconForProperty(const FProperty& InPropert
 	return false;
 }
 
+bool FCachedBindingData::GetPinTypeAndIconForProperty(const FProperty& InProperty, FPropertyBindingDataView InTargetDataView
+	, FEdGraphPinType& OutPinType, const FSlateBrush*& OutIconBrush) const
+{
+	return false;
+}
+
 bool FCachedBindingData::IsPropertyReference(const FProperty& InProperty)
 {
 	return false;
@@ -521,7 +527,7 @@ void FCachedBindingData::UpdateData()
 	check(Schema);
 
 	FEdGraphPinType PinType;
-	FName IconName;
+	const FSlateBrush* Icon = nullptr;
 	Schema->ConvertPropertyToPinType(Property, PinType);
 
 	FPropertyBindingDataView TargetDataView;
@@ -529,11 +535,11 @@ void FCachedBindingData::UpdateData()
 
 	const bool bIsPropertyReference = IsPropertyReference(*Property);
 
-	const bool bFound = GetPinTypeAndIconForProperty(*Property, TargetDataView, PinType, IconName);
+	const bool bFound = GetPinTypeAndIconForProperty(*Property, TargetDataView, PinType, Icon);
 	if (!bFound)
 	{
 		Schema->ConvertPropertyToPinType(Property, PinType);
-		IconName = TEXT("Kismet.Tabs.Variables");
+		Icon = FAppStyle::GetBrush("Kismet.Tabs.Variables");
 	}
 
 	FTextBuilder TooltipBuilder;
@@ -622,7 +628,7 @@ void FCachedBindingData::UpdateData()
 					Schema->ConvertPropertyToPinType(SourceLeafProperty, PinType);
 				}
 
-				Image = FAppStyle::GetBrush(IconName);
+				Image = Icon;
 				Color = Schema->GetPinTypeColor(PinType);
 
 				bIsCachedDataValid = true;
@@ -679,7 +685,7 @@ void FCachedBindingData::UpdateData()
 		FormatableText = FText::GetEmpty();
 		TooltipBuilder.AppendLineFormat(LOCTEXT("BindTooltip", "Bind {0} to value from another property."), UE::PropertyBinding::GetPropertyTypeText(Property));
 
-		Image = FAppStyle::GetBrush(IconName);
+		Image = Icon;
 		Color = Schema->GetPinTypeColor(PinType);
 
 		bIsCachedDataValid = true;
@@ -1435,6 +1441,8 @@ void FPropertyBindingExtension::ExtendWidgetRow(FDetailWidgetRow& InWidgetRow, c
 			false);
 	}
 
+	CustomizeWidgetArgs(Args, BindingsOwner, TargetPath, CachedBindingData);
+
 	InWidgetRow.ExtensionContent()
 	[
 		PropertyAccessEditor.MakePropertyBindingWidget(BindingContextStructs, Args)
diff --git a/Engine/Plugins/Runtime/PropertyBindingUtils/Source/PropertyBindingUtilsEditor/Public/PropertyBindingExtension.h b/Engine/Plugins/Runtime/PropertyBindingUtils/Source/PropertyBindingUtilsEditor/Public/PropertyBindingExtension.h
--- a/Engine/Plugins/Runtime/PropertyBindingUtils/Source/PropertyBindingUtilsEditor/Public/PropertyBindingExtension.h
+++ b/Engine/Plugins/Runtime/PropertyBindingUtils/Source/PropertyBindingUtilsEditor/Public/PropertyBindingExtension.h
@@ -20,6 +20,7 @@ struct FPropertyBindingBindableStructDescriptor;
 struct FPropertyBindingDataView;
 struct FPropertyBindingPath;
 struct FPropertyBindingBinding;
+struct FPropertyBindingWidgetArgs;
 
 namespace UE::PropertyBinding
 {
@@ -120,7 +121,10 @@ protected:
 	
 	PROPERTYBINDINGUTILSEDITOR_API virtual void AddPropertyInfoOverride(const FProperty& InProperty, TArray<TSharedPtr<const UE::PropertyBinding::FPropertyInfoOverride>>& OutPropertyInfoOverrides) const;
 	PROPERTYBINDINGUTILSEDITOR_API virtual void GetSourceDataViewForNewBinding(TNotNull<IPropertyBindingBindingCollectionOwner*> InBindingsOwner,TConstStructView<FPropertyBindingBindableStructDescriptor> InDescriptor, FPropertyBindingDataView& OutSourceDataView);
-	[[nodiscard]] PROPERTYBINDINGUTILSEDITOR_API virtual bool GetPinTypeAndIconForProperty(const FProperty& InProperty, FPropertyBindingDataView InTargetDataView, FEdGraphPinType& OutPinType, FName& OutIconName) const;
+
+	UE_DEPRECATED(5.7, "Use the overload with OutIconBrush instead")
+	[[nodiscard]] PROPERTYBINDINGUTILSEDITOR_API virtual bool GetPinTypeAndIconForProperty(const FProperty& InProperty, FPropertyBindingDataView InTargetDataView, FEdGraphPinType& OutPinType, FName& OutIconName) const final;
+	[[nodiscard]] PROPERTYBINDINGUTILSEDITOR_API virtual bool GetPinTypeAndIconForProperty(const FProperty& InProperty, FPropertyBindingDataView InTargetDataView, FEdGraphPinType& OutPinType, const FSlateBrush*& OutIconBrush) const;
 	[[nodiscard]] PROPERTYBINDINGUTILSEDITOR_API virtual bool GetPropertyFunctionText(FConstStructView InPropertyFunctionStructView, FText& OutText) const;
 	[[nodiscard]] PROPERTYBINDINGUTILSEDITOR_API virtual bool GetPropertyFunctionTooltipText(FConstStructView InPropertyFunctionStructView, FText& OutText) const;
 	[[nodiscard]] PROPERTYBINDINGUTILSEDITOR_API virtual bool GetPropertyFunctionIconColor(FConstStructView InPropertyFunctionStructView, FLinearColor& OutColor) const;
@@ -134,9 +138,10 @@ protected:
 		const void* InTargetPropertyValue,
 		bool& bOutAreCompatible) const;
 
+	PROPERTYBINDINGUTILSEDITOR_API void UpdateData();
+
 private:
 	PROPERTYBINDINGUTILSEDITOR_API void ConditionallyUpdateData();
-	PROPERTYBINDINGUTILSEDITOR_API void UpdateData();
 	PROPERTYBINDINGUTILSEDITOR_API bool ArePropertyAndContextStructCompatible(const UStruct* SourceStruct, const FProperty* TargetProperty) const;
 	PROPERTYBINDINGUTILSEDITOR_API bool ArePropertiesCompatible(const FProperty* SourceProperty, const FProperty* TargetProperty, const void* SourcePropertyValue, const void* TargetPropertyValue) const;
 
@@ -233,6 +238,7 @@ protected:
 	PROPERTYBINDINGUTILSEDITOR_API virtual uint64 GetDisallowedPropertyFlags() const;
 	PROPERTYBINDINGUTILSEDITOR_API virtual bool GetPromotionToParameterOverrideInternal(const FProperty& InProperty, bool& bOutOverride) const;
 	PROPERTYBINDINGUTILSEDITOR_API virtual void UpdateContextStruct(TConstStructView<FPropertyBindingBindableStructDescriptor> InStructDesc, FBindingContextStruct& InOutContextStruct, TMap<FString, FText>& InOutSectionNames) const;
+	PROPERTYBINDINGUTILSEDITOR_API virtual void CustomizeWidgetArgs(FPropertyBindingWidgetArgs& Args, IPropertyBindingBindingCollectionOwner* InBindingsOwner, const FPropertyBindingPath& InTargetPath, TSharedPtr<UE::PropertyBinding::FCachedBindingData> InCachedBindingData) { }
 	virtual bool CanBindToArrayElements() const
 	{
 		return false;
diff --git a/Engine/Plugins/Runtime/StateTree/Source/StateTreeEditorModule/Private/Customizations/StateTreeBindingExtension.cpp b/Engine/Plugins/Runtime/StateTree/Source/StateTreeEditorModule/Private/Customizations/StateTreeBindingExtension.cpp
--- a/Engine/Plugins/Runtime/StateTree/Source/StateTreeEditorModule/Private/Customizations/StateTreeBindingExtension.cpp
+++ b/Engine/Plugins/Runtime/StateTree/Source/StateTreeEditorModule/Private/Customizations/StateTreeBindingExtension.cpp
@@ -173,7 +173,7 @@ struct FStateTreeCachedBindingData : UE::PropertyBinding::FCachedBindingData, TS
 		}
 	}
 
-	virtual bool GetPinTypeAndIconForProperty(const FProperty& InProperty, FPropertyBindingDataView InTargetDataView, FEdGraphPinType& OutPinType, FName& OutIconName) const override
+	virtual bool GetPinTypeAndIconForProperty(const FProperty& InProperty, FPropertyBindingDataView InTargetDataView, FEdGraphPinType& OutPinType, const FSlateBrush*& OutIconBrush) const override
 	{
 		const bool bIsPropertyRef = PropertyRefHelpers::IsPropertyRef(InProperty);;
 		if (bIsPropertyRef && InTargetDataView.IsValid())
@@ -185,14 +185,14 @@ struct FStateTreeCachedBindingData : UE::PropertyBinding::FCachedBindingData, TS
 				const void* PropertyRef = TargetIndirections.Last().GetPropertyAddress();
 				OutPinType = PropertyRefHelpers::GetPropertyRefInternalTypeAsPin(InProperty, PropertyRef);
 			}
-			OutIconName = TEXT("Kismet.Tabs.Variables");
+			OutIconBrush = FAppStyle::GetBrush("Kismet.Tabs.Variables");
 			return true;
 		}
 
 		if (IsDelegateListenerProperty(&InProperty))
 		{
 			OutPinType.PinCategory = UEdGraphSchema_K2::PC_Delegate;
-			OutIconName = TEXT("Icons.Event");
+			OutIconBrush = FAppStyle::GetBrush("Icons.Event");
 			return true;
 		}
 
