[PropertyBindings] updated RemoveInvalidBindings to handle redirects in order to prevent undesired removal of still valid bindings

diff --git a/Engine/Plugins/Runtime/PropertyBindingUtils/Source/PropertyBindingUtils/Private/PropertyBindingBindingCollection.cpp b/Engine/Plugins/Runtime/PropertyBindingUtils/Source/PropertyBindingUtils/Private/PropertyBindingBindingCollection.cpp
--- a/Engine/Plugins/Runtime/PropertyBindingUtils/Source/PropertyBindingUtils/Private/PropertyBindingBindingCollection.cpp
+++ b/Engine/Plugins/Runtime/PropertyBindingUtils/Source/PropertyBindingUtils/Private/PropertyBindingBindingCollection.cpp
@@ -204,18 +204,19 @@ void FPropertyBindingBindingCollection::RemoveInvalidBindings(const TMap<FGuid,
 				return true;
 			}
 
+			constexpr bool bHandleRedirects = true;
+
 			// Remove binding if path containing instanced indirections (e.g. instance struct or object) cannot be resolved.
-			// TODO: Try to use core redirect to find new name.
 			{
 				const FPropertyBindingDataView* SourceValue = InValidStructs.Find(Binding.GetSourcePath().GetStructID());
 				if (SourceValue && SourceValue->IsValid())
 				{
 					FString Error;
 					TArray<FPropertyBindingPathIndirection> Indirections;
-					if (!Binding.GetSourcePath().ResolveIndirectionsWithValue(*SourceValue, Indirections, &Error))
+					if (!Binding.GetSourcePath().ResolveIndirectionsWithValue(*SourceValue, Indirections, &Error, bHandleRedirects))
 					{
-						PROPERTY_BINDINGS_LOG(LogOwner, Log, TEXT("Removing binding because source path '%s' cannot be resolved: %s"),
-							*Binding.GetSourcePath().ToString(), *Error); // Error contains the target path.
+						PROPERTY_BINDINGS_LOG(LogOwner, Log, TEXT("Removing binding from '%s' because source path '%s' cannot be resolved: %s"),
+							*GetPathNameSafe(LogOwner), *Binding.GetSourcePath().ToString(), *Error); // Error contains the source path.
 
 						// Remove
 						return true;
@@ -227,10 +228,10 @@ void FPropertyBindingBindingCollection::RemoveInvalidBindings(const TMap<FGuid,
 				const FPropertyBindingDataView TargetValue = InValidStructs.FindChecked(Binding.GetTargetPath().GetStructID());
 				FString Error;
 				TArray<FPropertyBindingPathIndirection> Indirections;
-				if (!Binding.GetTargetPath().ResolveIndirectionsWithValue(TargetValue, Indirections, &Error))
+				if (!Binding.GetTargetPath().ResolveIndirectionsWithValue(TargetValue, Indirections, &Error, bHandleRedirects))
 				{
-					PROPERTY_BINDINGS_LOG(LogOwner, Log, TEXT("Removing binding because target path '%s' cannot be resolved: %s"),
-						*Binding.GetSourcePath().ToString(), *Error); // Error contains the target path.
+					PROPERTY_BINDINGS_LOG(LogOwner, Log, TEXT("Removing binding from '%s' because target path '%s' cannot be resolved: %s"),
+						*GetPathNameSafe(LogOwner), *Binding.GetTargetPath().ToString(), *Error); // Error contains the target path.
 
 					// Remove
 					return true;
