Property Binding: fixed ContainsAnyStruct to return the correct result when checking for the struct in paths

diff --git a/Engine/Plugins/Runtime/PropertyBindingUtils/Source/PropertyBindingUtils/Private/PropertyBindingBindingCollection.cpp b/Engine/Plugins/Runtime/PropertyBindingUtils/Source/PropertyBindingUtils/Private/PropertyBindingBindingCollection.cpp
--- a/Engine/Plugins/Runtime/PropertyBindingUtils/Source/PropertyBindingUtils/Private/PropertyBindingBindingCollection.cpp
+++ b/Engine/Plugins/Runtime/PropertyBindingUtils/Source/PropertyBindingUtils/Private/PropertyBindingBindingCollection.cpp
@@ -1257,7 +1257,7 @@ bool FPropertyBindingBindingCollection::ContainsAnyStruct(const TSet<const UStru
 			}
 			return EVisitResult::Continue;
 		});
-	return false;
+	return bPathContainsStruct;
 }
 
 #if WITH_EDITOR || WITH_PROPERTYBINDINGUTILS_DEBUG
