[StateTree] Get the latest name of the binding for the editor binding view instead of caching the value. Prevent issues when the variable is renamed.

diff --git a/Engine/Plugins/Runtime/PropertyBindingUtils/Source/PropertyBindingUtilsEditor/Private/Widgets/SPropertyBindingView.cpp b/Engine/Plugins/Runtime/PropertyBindingUtils/Source/PropertyBindingUtilsEditor/Private/Widgets/SPropertyBindingView.cpp
--- a/Engine/Plugins/Runtime/PropertyBindingUtils/Source/PropertyBindingUtilsEditor/Private/Widgets/SPropertyBindingView.cpp
+++ b/Engine/Plugins/Runtime/PropertyBindingUtils/Source/PropertyBindingUtilsEditor/Private/Widgets/SPropertyBindingView.cpp
@@ -45,6 +45,12 @@ namespace Private
 			return FText::GetEmpty();
 		}
 
+		FText HandleGetBindingPath(bool bIsSource) const
+		{
+			const FPropertyBindingPath& Path = bIsSource ? Item->SourcePath : Item->TargetPath;
+			return FText::FromString(Path.ToString());
+		}
+
 		virtual TSharedRef<SWidget> GenerateWidgetForColumn(const FName& ColumnName) override
 		{
 			if (ColumnName == ColumnId_SourceStruct)
@@ -56,7 +62,7 @@ namespace Private
 			else if (ColumnName == ColumnId_SourcePath)
 			{
 				return SNew(STextBlock)
-					.Text(FText::FromString(Item->SourcePath.ToString()));		
+					.Text(this, &SBindingViewRow::HandleGetBindingPath, true);
 			}
 			else if (ColumnName == ColumnId_TargetStruct)
 			{
@@ -67,7 +73,7 @@ namespace Private
 			else if (ColumnName == ColumnId_TargetPath)
 			{
 				return SNew(STextBlock)
-					.Text(FText::FromString(Item->TargetPath.ToString()));
+					.Text(this, &SBindingViewRow::HandleGetBindingPath, false);
 			}
 
 			return SNullWidget::NullWidget;
