[PropertyBindingUtils] Quick optimization to prevent dynamic string creation.

diff --git a/Engine/Plugins/Runtime/PropertyBindingUtils/Source/PropertyBindingUtils/Private/PropertyBindingPath.cpp b/Engine/Plugins/Runtime/PropertyBindingUtils/Source/PropertyBindingUtils/Private/PropertyBindingPath.cpp
--- a/Engine/Plugins/Runtime/PropertyBindingUtils/Source/PropertyBindingUtils/Private/PropertyBindingPath.cpp
+++ b/Engine/Plugins/Runtime/PropertyBindingUtils/Source/PropertyBindingUtils/Private/PropertyBindingPath.cpp
@@ -136,16 +136,18 @@ FString FPropertyBindingPath::ToString(const int32 HighlightedSegment, const TCH
 	{
 		if (Segment.GetIndex() > 0)
 		{
-			Result += TEXT(".");
+			Result << TEXT('.');
 		}
 		if (Segment.GetIndex() == HighlightedSegment && HighlightPrefix)
 		{
-			Result += HighlightPrefix;
+			Result << HighlightPrefix;
 		}
 
 		if (bOutputInstances && Segment->GetInstanceStruct())
 		{
-			Result += FString::Printf(TEXT("(%s)"), *GetNameSafe(Segment->GetInstanceStruct()));
+			Result << TEXT('(');
+			Result << Segment->GetInstanceStruct()->GetFName();
+			Result << TEXT(')');
 		}
 
 #if WITH_EDITORONLY_DATA
@@ -156,27 +158,29 @@ FString FPropertyBindingPath::ToString(const int32 HighlightedSegment, const TCH
 			const FString FriendlyName = FStructureEditorUtils::GetVariableFriendlyName(ParentUserDefinedStruct, Segment->GetPropertyGuid());
 			if (!FriendlyName.IsEmpty())
 			{
-				Result += FriendlyName;
+				Result << FriendlyName;
 			}
 			else
 			{
-				Result += Segment->GetName().ToString();
+				Result << Segment->GetName();
 			}
 		}
 		else
 #endif // WITH_EDITORONLY_DATA
 		{
-			Result += Segment->GetName().ToString();
+			Result << Segment->GetName();
 		}
 
 		if (Segment->GetArrayIndex() >= 0)
 		{
-			Result += FString::Printf(TEXT("[%d]"), Segment->GetArrayIndex());
+			Result << TEXT('[');
+			Result << Segment->GetArrayIndex();
+			Result << TEXT(']');
 		}
 
 		if (Segment.GetIndex() == HighlightedSegment && HighlightPostfix)
 		{
-			Result += HighlightPostfix;
+			Result << HighlightPostfix;
 		}
 	}
 	return Result.ToString();
