[State Tree] fixed behavior of removing duplicate bindings. It should spare one instead of removing all

diff --git a/Engine/Plugins/Runtime/PropertyBindingUtils/Source/PropertyBindingUtils/Private/PropertyBindingBindingCollection.cpp b/Engine/Plugins/Runtime/PropertyBindingUtils/Source/PropertyBindingUtils/Private/PropertyBindingBindingCollection.cpp
--- a/Engine/Plugins/Runtime/PropertyBindingUtils/Source/PropertyBindingUtils/Private/PropertyBindingBindingCollection.cpp
+++ b/Engine/Plugins/Runtime/PropertyBindingUtils/Source/PropertyBindingUtils/Private/PropertyBindingBindingCollection.cpp
@@ -311,11 +311,17 @@ void FPropertyBindingBindingCollection::RemoveInvalidBindings(const TMap<FGuid,
 			}
 		}
 
-		for (const FPropertyBindingPath& BindingPath : BindingPathsToRemove)
+		RemoveBindingsInternal([&BindingPathsToRemove](FPropertyBindingBinding& InBinding)
 		{
-			constexpr ESearchMode SearchMode = ESearchMode::Exact;
-			RemoveBindings(BindingPath, SearchMode);
-		}
+			int32 FoundIndex = BindingPathsToRemove.Find(InBinding.GetTargetPath());
+			if (FoundIndex == INDEX_NONE)
+			{
+				return false;
+			}
+
+			BindingPathsToRemove.RemoveAtSwap(FoundIndex);
+			return true;
+		});
 
 		++MapIt;
 	}
