[State Tree] fixed fallout from CL 38662366 that PropertyRef promotion to parameter was disabled

diff --git a/Engine/Plugins/Runtime/PropertyBindingUtils/Source/PropertyBindingUtilsEditor/Public/PropertyBindingExtension.h b/Engine/Plugins/Runtime/PropertyBindingUtils/Source/PropertyBindingUtilsEditor/Public/PropertyBindingExtension.h
--- a/Engine/Plugins/Runtime/PropertyBindingUtils/Source/PropertyBindingUtilsEditor/Public/PropertyBindingExtension.h
+++ b/Engine/Plugins/Runtime/PropertyBindingUtils/Source/PropertyBindingUtilsEditor/Public/PropertyBindingExtension.h
@@ -237,6 +237,13 @@ protected:
 	PROPERTYBINDINGUTILSEDITOR_API virtual TSharedPtr<UE::PropertyBinding::FCachedBindingData> CreateCachedBindingData(IPropertyBindingBindingCollectionOwner* InBindingsOwner, const FPropertyBindingPath& InTargetPath, const TSharedPtr<IPropertyHandle>& InPropertyHandle, TConstArrayView<TInstancedStruct<FPropertyBindingBindableStructDescriptor>> InAccessibleStructs) const;
 	PROPERTYBINDINGUTILSEDITOR_API virtual bool CanBindToProperty(const FPropertyBindingPath& InTargetPath, const IPropertyHandle& InPropertyHandle) const;
 	PROPERTYBINDINGUTILSEDITOR_API virtual uint64 GetDisallowedPropertyFlags() const;
+
+	/**
+	 * The optional override criteria of whether the property can be promoted to parameter 
+	 * @param InProperty the property to be tested upon
+	 * @param bOutOverride Whether the property should be promoted to parameter, if the return value is true
+	 * @return Whether the promotion criteria is overridden for this property
+	 */
 	PROPERTYBINDINGUTILSEDITOR_API virtual bool GetPromotionToParameterOverrideInternal(const FProperty& InProperty, bool& bOutOverride) const;
 	PROPERTYBINDINGUTILSEDITOR_API virtual void UpdateContextStruct(TConstStructView<FPropertyBindingBindableStructDescriptor> InStructDesc, FBindingContextStruct& InOutContextStruct, TMap<FString, FText>& InOutSectionNames) const;
 	PROPERTYBINDINGUTILSEDITOR_API virtual void CustomizeDetailWidgetRow(FDetailWidgetRow& InWidgetRow, const IDetailLayoutBuilder& InDetailBuilder, IPropertyBindingBindingCollectionOwner* InBindingsOwner, TSharedPtr<IPropertyHandle> InPropertyHandle, const FPropertyBindingPath& InTargetPath, TSharedPtr<UE::PropertyBinding::FCachedBindingData> InCachedBindingData) const;
diff --git a/Engine/Plugins/Runtime/StateTree/Source/StateTreeEditorModule/Private/Customizations/StateTreeBindingExtension.cpp b/Engine/Plugins/Runtime/StateTree/Source/StateTreeEditorModule/Private/Customizations/StateTreeBindingExtension.cpp
--- a/Engine/Plugins/Runtime/StateTree/Source/StateTreeEditorModule/Private/Customizations/StateTreeBindingExtension.cpp
+++ b/Engine/Plugins/Runtime/StateTree/Source/StateTreeEditorModule/Private/Customizations/StateTreeBindingExtension.cpp
@@ -822,10 +822,11 @@ bool FStateTreeBindingExtension::GetPromotionToParameterOverrideInternal(const F
 		// Support Property Refs as even though these aren't bp types, the actual types that would be added are the ones in the meta-data RefType
 		if (StructProperty->Struct && StructProperty->Struct->IsChildOf(FStateTreePropertyRef::StaticStruct()))
 		{
-			bOutOverride = false;
+			bOutOverride = true;
 			return true;
 		}
 	}
+
 	return false;
 }
 
diff --git a/Engine/Plugins/VirtualProduction/SceneState/Source/SceneStateBlueprintEditor/Private/DetailsView/SceneStatePropertyBindingExtension.cpp b/Engine/Plugins/VirtualProduction/SceneState/Source/SceneStateBlueprintEditor/Private/DetailsView/SceneStatePropertyBindingExtension.cpp
--- a/Engine/Plugins/VirtualProduction/SceneState/Source/SceneStateBlueprintEditor/Private/DetailsView/SceneStatePropertyBindingExtension.cpp
+++ b/Engine/Plugins/VirtualProduction/SceneState/Source/SceneStateBlueprintEditor/Private/DetailsView/SceneStatePropertyBindingExtension.cpp
@@ -63,7 +63,7 @@ bool FBindingExtension::GetPromotionToParameterOverrideInternal(const FProperty&
 		// Support Property Refs as even though these aren't bp types, the actual types that would be added are the ones in the meta-data RefType
 		if (StructProperty->Struct && StructProperty->Struct->IsChildOf<FSceneStatePropertyReference>())
 		{
-			bOutOverride = false;
+			bOutOverride = true;
 			return true;
 		}
 	}
