PropertyBindings: Improve the performance of ResolvePaths modestly by eliminating some low-hanging fruits.

diff --git a/Engine/Plugins/Runtime/PropertyBindingUtils/Source/PropertyBindingUtils/Private/PropertyBindingBindingCollection.cpp b/Engine/Plugins/Runtime/PropertyBindingUtils/Source/PropertyBindingUtils/Private/PropertyBindingBindingCollection.cpp
--- a/Engine/Plugins/Runtime/PropertyBindingUtils/Source/PropertyBindingUtils/Private/PropertyBindingBindingCollection.cpp
+++ b/Engine/Plugins/Runtime/PropertyBindingUtils/Source/PropertyBindingUtils/Private/PropertyBindingBindingCollection.cpp
@@ -256,6 +256,8 @@ void FPropertyBindingBindingCollection::Reset()
 
 bool FPropertyBindingBindingCollection::ResolvePaths()
 {
+	TRACE_CPUPROFILER_EVENT_SCOPE(FPropertyBindingBindingCollection::ResolvePaths)
+
 	PropertyIndirections.Reset();
 	PropertyCopies.SetNum(GetNumBindings());
 
@@ -384,7 +386,7 @@ bool FPropertyBindingBindingCollection::ResolvePath(const FPropertyBindingDataVi
 	}
 
 	FString Error;
-	TArray<FPropertyBindingPathIndirection> PathIndirections;
+	TArray<FPropertyBindingPathIndirection, TInlineAllocator<16>> PathIndirections;
 	if (!Path.ResolveIndirectionsWithValue(DataView, PathIndirections, &Error))
 	{
 		PROPERTY_BINDINGS_LOG(GetLogOwner(), Error, TEXT("%hs: %s"), __FUNCTION__, *Error);
@@ -1092,6 +1094,8 @@ bool FPropertyBindingBindingCollection::PerformCopy(const FPropertyBindingCopyIn
 
 bool FPropertyBindingBindingCollection::CopyProperty(const FPropertyBindingCopyInfo& Copy, FPropertyBindingDataView SourceStructView, FPropertyBindingDataView TargetStructView) const
 {
+	TRACE_CPUPROFILER_EVENT_SCOPE(FPropertyBindingBindingCollection::CopyProperty)
+
 	// This is made ensure so that the programmers have the chance to catch it (it's usually programming error not to call ResolvePaths(), and it wont spam log for others.
 	if (!ensureMsgf(bBindingsResolved, TEXT("Bindings must be resolved successfully before copying. See ResolvePaths()")))
 	{
diff --git a/Engine/Plugins/Runtime/PropertyBindingUtils/Source/PropertyBindingUtils/Private/PropertyBindingPath.cpp b/Engine/Plugins/Runtime/PropertyBindingUtils/Source/PropertyBindingUtils/Private/PropertyBindingPath.cpp
--- a/Engine/Plugins/Runtime/PropertyBindingUtils/Source/PropertyBindingUtils/Private/PropertyBindingPath.cpp
+++ b/Engine/Plugins/Runtime/PropertyBindingUtils/Source/PropertyBindingUtils/Private/PropertyBindingPath.cpp
@@ -216,7 +216,8 @@ bool FPropertyBindingPath::ResolveIndirections(const UStruct* BaseStruct, TArray
 	return ResolveIndirectionsWithValue(FPropertyBindingDataView(BaseStruct, nullptr), OutIndirections, OutError, bHandleRedirects);
 }
 
-bool FPropertyBindingPath::ResolveIndirectionsWithValue(const FPropertyBindingDataView BaseValueView, TArray<FPropertyBindingPathIndirection>& OutIndirections, FString* OutError, bool bHandleRedirects) const
+template<typename Allocator>
+bool FPropertyBindingPath::ResolveIndirectionsWithValue(const FPropertyBindingDataView BaseValueView, TArray<FPropertyBindingPathIndirection, Allocator>& OutIndirections, FString* OutError, bool bHandleRedirects) const
 {
 	OutIndirections.Reset();
 	if (OutError)
@@ -605,6 +606,16 @@ bool FPropertyBindingPath::ResolveIndirectionsWithValue(const FPropertyBindingDa
 	return true;
 }
 
+bool FPropertyBindingPath::ResolveIndirectionsWithValue(const FPropertyBindingDataView BaseValueView, TArray<FPropertyBindingPathIndirection>& OutIndirections, FString* OutError, bool bHandleRedirects) const
+{
+	return ResolveIndirectionsWithValue<>(BaseValueView, OutIndirections, OutError, bHandleRedirects);
+}
+
+bool FPropertyBindingPath::ResolveIndirectionsWithValue(const FPropertyBindingDataView BaseValueView, TArray<FPropertyBindingPathIndirection, TInlineAllocator<16>>& OutIndirections, FString* OutError, bool bHandleRedirects) const
+{
+	return ResolveIndirectionsWithValue<>(BaseValueView, OutIndirections, OutError, bHandleRedirects);
+}
+
 bool FPropertyBindingPath::operator==(const FPropertyBindingPath& RHS) const
 {
 #if WITH_EDITORONLY_DATA
diff --git a/Engine/Plugins/Runtime/PropertyBindingUtils/Source/PropertyBindingUtils/Public/PropertyBindingPath.h b/Engine/Plugins/Runtime/PropertyBindingUtils/Source/PropertyBindingUtils/Public/PropertyBindingPath.h
--- a/Engine/Plugins/Runtime/PropertyBindingUtils/Source/PropertyBindingUtils/Public/PropertyBindingPath.h
+++ b/Engine/Plugins/Runtime/PropertyBindingUtils/Source/PropertyBindingUtils/Public/PropertyBindingPath.h
@@ -302,6 +302,7 @@ struct FPropertyBindingPath
 	 * @return true of path could be resolved against the base value, and instance types were updated.
 	 */
 	PROPERTYBINDINGUTILS_API bool ResolveIndirectionsWithValue(const FPropertyBindingDataView BaseValueView, TArray<FPropertyBindingPathIndirection>& OutIndirections, FString* OutError = nullptr, bool bHandleRedirects = false) const;
+	PROPERTYBINDINGUTILS_API bool ResolveIndirectionsWithValue(const FPropertyBindingDataView BaseValueView, TArray<FPropertyBindingPathIndirection, TInlineAllocator<16>>& OutIndirections, FString* OutError = nullptr, bool bHandleRedirects = false) const;
 
 	/**
 	 * Updates property segments from base struct type. The path is expected to be relative to the BaseStruct.
@@ -412,6 +413,9 @@ private:
 	/** Path segments pointing to a specific property on the path. */
 	UPROPERTY(VisibleAnywhere, Category = "Bindings")
 	TArray<FPropertyBindingPathSegment> Segments;
+
+	template<typename Allocator>
+	bool ResolveIndirectionsWithValue(const FPropertyBindingDataView BaseValueView, TArray<FPropertyBindingPathIndirection, Allocator>& OutIndirections, FString* OutError = nullptr, bool bHandleRedirects = false) const;
 };
 
 
diff --git a/Engine/Source/Runtime/CoreUObject/Private/StructUtils/InstancedStructContainer.cpp b/Engine/Source/Runtime/CoreUObject/Private/StructUtils/InstancedStructContainer.cpp
--- a/Engine/Source/Runtime/CoreUObject/Private/StructUtils/InstancedStructContainer.cpp
+++ b/Engine/Source/Runtime/CoreUObject/Private/StructUtils/InstancedStructContainer.cpp
@@ -89,6 +89,7 @@ FInstancedStructContainer& FInstancedStructContainer::operator=(TConstArrayView<
 void FInstancedStructContainer::Append(const FInstancedStructContainer& Other)
 {
 	TArray<FConstStructView> Views;
+	Views.Reserve(Other.Num());
 	for (int32 Index = 0; Index < Other.Num(); Index++)
 	{
 		Views.Add(Other[Index]);
@@ -115,6 +116,7 @@ void FInstancedStructContainer::Append(TConstArrayView<FConstStructView> NewItem
 void FInstancedStructContainer::InsertAt(const int32 InsertAtIndex, const FInstancedStructContainer& Other)
 {
 	TArray<FConstStructView> Views;
+	Views.Reserve(Other.Num());
 	for (int32 Index = 0; Index < Other.Num(); Index++)
 	{
 		Views.Add(Other[Index]);
diff --git a/Engine/Source/Runtime/CoreUObject/Public/StructUtils/InstancedStruct.h b/Engine/Source/Runtime/CoreUObject/Public/StructUtils/InstancedStruct.h
--- a/Engine/Source/Runtime/CoreUObject/Public/StructUtils/InstancedStruct.h
+++ b/Engine/Source/Runtime/CoreUObject/Public/StructUtils/InstancedStruct.h
@@ -287,6 +287,7 @@ struct TStructOpsTypeTraits<FInstancedStruct> : public TStructOpsTypeTraitsBase2
 		WithFindInnerPropertyInstance = true,
 		WithClearOnFinishDestroy = true,
 		WithVisitor = true,
+		WithZeroConstructor = true,
 	};
 };
 
