PropertyBindings: Introduce BindingRoot with which it's possible to specify a property to be used as the root of a binding (instead of following the parent chain until the top)

#rnx

diff --git a/Engine/Plugins/Runtime/PropertyBindingUtils/Source/PropertyBindingUtilsEditor/Private/PropertyBindingExtension.cpp b/Engine/Plugins/Runtime/PropertyBindingUtils/Source/PropertyBindingUtilsEditor/Private/PropertyBindingExtension.cpp
--- a/Engine/Plugins/Runtime/PropertyBindingUtils/Source/PropertyBindingUtilsEditor/Private/PropertyBindingExtension.cpp
+++ b/Engine/Plugins/Runtime/PropertyBindingUtils/Source/PropertyBindingUtilsEditor/Private/PropertyBindingExtension.cpp
@@ -150,6 +150,17 @@ TSharedPtr<const IPropertyHandle> MakeStructPropertyPathFromPropertyHandle(
 	FGuid StructID;
 	TArray<FPropertyBindingPathSegment> PathSegments;
 
+	auto ExtractStructIdFromPropertyHandle = [&StructID](const TSharedPtr<const IPropertyHandle>& ParentHandle)
+	{
+		if (const FString* IDString = ParentHandle->GetInstanceMetaData(MetaDataStructIDName))
+		{
+			LexFromString(StructID, **IDString);
+			return true;
+		}
+
+		return false;
+	};
+
 	TSharedPtr<const IPropertyHandle> BindablePropertyHandle;
 	TSharedPtr<const IPropertyHandle> CurrentPropertyHandle = InPropertyHandle;
 	while (CurrentPropertyHandle.IsValid())
@@ -223,15 +234,29 @@ TSharedPtr<const IPropertyHandle> MakeStructPropertyPathFromPropertyHandle(
 			}
 
 			// Bindable property must have node ID
-			if (const FString* IDString = CurrentPropertyHandle->GetInstanceMetaData(MetaDataStructIDName))
+			if (ExtractStructIdFromPropertyHandle(CurrentPropertyHandle))
 			{
-				LexFromString(StructID, **IDString);
 				BindablePropertyHandle = CurrentPropertyHandle;
 				break;
 			}
 		}
+
+		TSharedPtr<IPropertyHandle> ParentHandle = CurrentPropertyHandle->GetParentHandle();
+		if (ParentHandle)
+		{
+			// Let's check if a property is acting as the root of the property path
+			if (ParentHandle->HasMetaData(MetaDataBindingRoot))
+			{
+				if (ExtractStructIdFromPropertyHandle(ParentHandle))
+				{
+					BindablePropertyHandle = CurrentPropertyHandle;
+				}
+
+				break;
+			}
+		}
 		
-		CurrentPropertyHandle = CurrentPropertyHandle->GetParentHandle();
+		CurrentPropertyHandle = MoveTemp(ParentHandle);
 	}
 
 	if (!StructID.IsValid() && InFallbackStructID.IsValid())
diff --git a/Engine/Plugins/Runtime/PropertyBindingUtils/Source/PropertyBindingUtilsEditor/Public/PropertyBindingExtension.h b/Engine/Plugins/Runtime/PropertyBindingUtils/Source/PropertyBindingUtilsEditor/Public/PropertyBindingExtension.h
--- a/Engine/Plugins/Runtime/PropertyBindingUtils/Source/PropertyBindingUtilsEditor/Public/PropertyBindingExtension.h
+++ b/Engine/Plugins/Runtime/PropertyBindingUtils/Source/PropertyBindingUtilsEditor/Public/PropertyBindingExtension.h
@@ -49,6 +49,13 @@ const FLazyName MetaDataStructIDName("StructIDForPropertyBinding");
 const FLazyName MetaDataNoBindingName("NoBinding");
 const FLazyName MetaDataNoPromoteToParameter("NoPromoteToParameter");
 
+/**
+ * Useful if you'd like to cut off some properties from the final path. For example, if your bindable struct descriptors
+ * point to structs, and you have these structs wrapped into object nodes on your graph, then you'd annotate the
+ * struct property inside the UObject node with both StructIDForPropertyBinding and BindingRoot.
+ */
+const FLazyName MetaDataBindingRoot("BindingRoot");
+
 PROPERTYBINDINGUTILSEDITOR_API extern FText GetPropertyTypeText(const FProperty* Property);
 
 /**
