[StateTree] Add a graph schema for StateTree. It will allow us to override the color and icon of pins.

diff --git a/Engine/Plugins/Runtime/PropertyBindingUtils/Source/PropertyBindingUtilsEditor/Private/PropertyBindingExtension.cpp b/Engine/Plugins/Runtime/PropertyBindingUtils/Source/PropertyBindingUtilsEditor/Private/PropertyBindingExtension.cpp
--- a/Engine/Plugins/Runtime/PropertyBindingUtils/Source/PropertyBindingUtilsEditor/Private/PropertyBindingExtension.cpp
+++ b/Engine/Plugins/Runtime/PropertyBindingUtils/Source/PropertyBindingUtilsEditor/Private/PropertyBindingExtension.cpp
@@ -329,10 +329,15 @@ void FMenuSectionHelper::SetSection(const FText& InSection)
 	}
 }
 
+// deprecated
 FText GetPropertyTypeText(const FProperty* Property)
+{
+	return GetPropertyTypeText(GetDefault<UEdGraphSchema_K2>(), Property);
+}
+
+FText GetPropertyTypeText(TNotNull<const UEdGraphSchema_K2*> Schema, const FProperty* Property)
 {
 	FEdGraphPinType PinType;
-	const UEdGraphSchema_K2* Schema = GetDefault<UEdGraphSchema_K2>();
 	Schema->ConvertPropertyToPinType(Property, PinType);
 				
 	const FName PinSubCategory = PinType.PinSubCategory;
@@ -349,10 +354,10 @@ FText GetPropertyTypeText(const FProperty* Property)
 	return UEdGraphSchema_K2::GetCategoryText(PinType.PinCategory, NAME_None, true);
 }
 
-TSharedRef<SWidget> MakeBindingPropertyInfoWidget(const FText& InDisplayText, const FEdGraphPinType& InPinType)
+TSharedRef<SWidget> MakeBindingPropertyInfoWidget(TNotNull<const UEdGraphSchema_K2*> Schema, const FText& InDisplayText, const FEdGraphPinType& InPinType)
 {
 	const FSlateBrush* Icon = FBlueprintEditorUtils::GetIconFromPin(InPinType, /*bIsLarge*/true);
-	const FLinearColor IconColor = GetDefault<UEdGraphSchema_K2>()->GetPinTypeColor(InPinType);
+	const FLinearColor IconColor = Schema->GetPinTypeColor(InPinType);
 
 	return SNew(SHorizontalBox)
 		+SHorizontalBox::Slot()
@@ -508,18 +513,30 @@ void FCachedBindingData::RemoveBinding(const FPropertyBindingBindingCollection::
 	BindingOwner->OnPropertyBindingChanged(SourcePath, TargetPath);
 }
 
+// deprecated
 bool FCachedBindingData::GetPinTypeAndIconForProperty(const FProperty& InProperty, FPropertyBindingDataView InTargetDataView
 	, FEdGraphPinType& OutPinType, FName& OutIconName) const
 {
 	return false;
 }
 
+// deprecated
 bool FCachedBindingData::GetPinTypeAndIconForProperty(const FProperty& InProperty, FPropertyBindingDataView InTargetDataView
 	, FEdGraphPinType& OutPinType, const FSlateBrush*& OutIconBrush) const
 {
 	return false;
 }
 
+FCachedBindingData::FPinInfo FCachedBindingData::GetPinInfoForProperty(const TNotNull<const FProperty*> InProperty, FPropertyBindingDataView InTargetDataView) const
+{
+	FPinInfo Info;
+	const UEdGraphSchema_K2* Schema = GetEdGraphSchema();
+	Schema->ConvertPropertyToPinType(InProperty, Info.Type);
+	Info.IconBrush = FAppStyle::GetBrush("Kismet.Tabs.Variables");
+	Info.IconColor = Schema->GetPinTypeColor(Info.Type);
+	return Info;
+}
+
 bool FCachedBindingData::IsPropertyReference(const FProperty& InProperty)
 {
 	return false;
@@ -565,24 +582,14 @@ void FCachedBindingData::UpdateData()
 		return;
 	}
 
-	const UEdGraphSchema_K2* Schema = GetDefault<UEdGraphSchema_K2>();
+	const UEdGraphSchema_K2* Schema = GetEdGraphSchema();
 	check(Schema);
 
-	FEdGraphPinType PinType;
-	const FSlateBrush* Icon = nullptr;
-	Schema->ConvertPropertyToPinType(Property, PinType);
-
 	FPropertyBindingDataView TargetDataView;
 	BindingOwner->GetBindingDataViewByID(TargetPath.GetStructID(), TargetDataView);
 
 	const bool bIsPropertyReference = IsPropertyReference(*Property);
-
-	const bool bFound = GetPinTypeAndIconForProperty(*Property, TargetDataView, PinType, Icon);
-	if (!bFound)
-	{
-		Schema->ConvertPropertyToPinType(Property, PinType);
-		Icon = FAppStyle::GetBrush("Kismet.Tabs.Variables");
-	}
+	FPinInfo PinInfo = GetPinInfoForProperty(Property, TargetDataView);
 
 	FTextBuilder TooltipBuilder;
 
@@ -665,13 +672,14 @@ void FCachedBindingData::UpdateData()
 
 				// Update pin type with source property so property reference that can bind
 				// to multiple types display the bound one.
-				if (bIsPropertyReference)
+				if (bIsPropertyReference && SourceLeafProperty)
 				{
-					Schema->ConvertPropertyToPinType(SourceLeafProperty, PinType);
+					Schema->ConvertPropertyToPinType(SourceLeafProperty, PinInfo.Type);
+					PinInfo.IconColor = Schema->GetPinTypeColor(PinInfo.Type);
 				}
 
-				Image = Icon;
-				Color = Schema->GetPinTypeColor(PinType);
+				Image = PinInfo.IconBrush;
+				Color = PinInfo.IconColor;
 
 				bIsCachedDataValid = true;
 			}
@@ -680,13 +688,13 @@ void FCachedBindingData::UpdateData()
 				FText SourceType;
 				if (SourceLeafProperty)
 				{
-					SourceType = GetPropertyTypeText(SourceLeafProperty);
+					SourceType = GetPropertyTypeText(Schema, SourceLeafProperty);
 				}
 				else if (SourceStruct)
 				{
 					SourceType = SourceStruct->GetDisplayNameText();
 				}
-				FText TargetType = GetPropertyTypeText(Property);
+				FText TargetType = GetPropertyTypeText(Schema, Property);
 
 				if (SourcePath.IsPathEmpty())
 				{
@@ -725,10 +733,10 @@ void FCachedBindingData::UpdateData()
 	{
 		// No bindings
 		FormatableText = FText::GetEmpty();
-		TooltipBuilder.AppendLineFormat(LOCTEXT("BindTooltip", "Bind {0} to value from another property."), UE::PropertyBinding::GetPropertyTypeText(Property));
+		TooltipBuilder.AppendLineFormat(LOCTEXT("BindTooltip", "Bind {0} to value from another property."), UE::PropertyBinding::GetPropertyTypeText(GetEdGraphSchema(), Property));
 
-		Image = Icon;
-		Color = Schema->GetPinTypeColor(PinType);
+		Image = PinInfo.IconBrush;
+		Color = PinInfo.IconColor;
 
 		bIsCachedDataValid = true;
 
@@ -1133,6 +1141,12 @@ const FSlateBrush* FCachedBindingData::GetImage()
 	return Image;
 }
 
+
+TNotNull<const UEdGraphSchema_K2*> FCachedBindingData::GetEdGraphSchema() const
+{
+	return GetDefault<UEdGraphSchema_K2>();
+}
+
 void FCachedBindingData::ConditionallyUpdateData()
 {
 	IPropertyBindingBindingCollectionOwner* BindingOwner = WeakBindingsOwner.Get();
@@ -1411,11 +1425,11 @@ void FPropertyBindingExtension::ExtendWidgetRow(FDetailWidgetRow& InWidgetRow, c
 										{
 											PromoteToParameterAndRefresh(PropertyName, TConstStructView<FPropertyBindingBindableStructDescriptor>(ContextStruct), nullptr);
 										}),
-									MakeContextStructWidget(ContextStruct));
+									MakeContextStructWidget(InCachedBindingData->GetEdGraphSchema(), ContextStruct));
 							}
 							else
 							{
-								InMenuBuilder.AddSubMenu(MakeContextStructWidget(ContextStruct),
+								InMenuBuilder.AddSubMenu(MakeContextStructWidget(InCachedBindingData->GetEdGraphSchema(), ContextStruct),
 									FNewMenuDelegate::CreateLambda([InCachedBindingData, PropertyName, ContextStruct, PropertyInfoOverrides = MoveTemp(PropertyInfoOverrides), PromoteToParameterAndRefresh](FMenuBuilder& InSubMenuBuilder)
 										{
 											FMenuSectionHelper SectionHelper(InSubMenuBuilder);
@@ -1428,7 +1442,7 @@ void FPropertyBindingExtension::ExtendWidgetRow(FDetailWidgetRow& InWidgetRow, c
 														{
 															PromoteToParameterAndRefresh(PropertyName, TConstStructView<FPropertyBindingBindableStructDescriptor>(ContextStruct), PropertyInfoOverride);
 														}),
-													MakeBindingPropertyInfoWidget(PropertyInfoOverride->TypeNameText, PropertyInfoOverride->PinType));
+													MakeBindingPropertyInfoWidget(InCachedBindingData->GetEdGraphSchema(), PropertyInfoOverride->TypeNameText, PropertyInfoOverride->PinType));
 											}
 										}));
 							}
@@ -1605,7 +1619,18 @@ bool FPropertyBindingExtension::CanBindToArrayElements() const
 	return false;
 }
 
+TNotNull<const UEdGraphSchema_K2*> FPropertyBindingExtension::GetEdGraphSchema() const
+{
+	return GetDefault<UEdGraphSchema_K2>();
+}
+
+// deprecated
 TSharedRef<SWidget> FPropertyBindingExtension::MakeContextStructWidget(const FPropertyBindingBindableStructDescriptor& InContextStruct)
+{
+	return MakeContextStructWidget(GetDefault<UEdGraphSchema_K2>(), InContextStruct);
+}
+
+TSharedRef<SWidget> FPropertyBindingExtension::MakeContextStructWidget(TNotNull<const UEdGraphSchema_K2*> Schema, const FPropertyBindingBindableStructDescriptor& InContextStruct)
 {
 	FEdGraphPinType PinType;
 
@@ -1625,7 +1650,7 @@ TSharedRef<SWidget> FPropertyBindingExtension::MakeContextStructWidget(const FPr
 	}
 
 	const FSlateBrush* Icon = FBlueprintEditorUtils::GetIconFromPin(PinType, true);
-	const FLinearColor IconColor = GetDefault<UEdGraphSchema_K2>()->GetPinTypeColor(PinType);
+	const FLinearColor IconColor = Schema->GetPinTypeColor(PinType);
 
 	return SNew(SHorizontalBox)
 		+ SHorizontalBox::Slot()
diff --git a/Engine/Plugins/Runtime/PropertyBindingUtils/Source/PropertyBindingUtilsEditor/Public/PropertyBindingExtension.h b/Engine/Plugins/Runtime/PropertyBindingUtils/Source/PropertyBindingUtilsEditor/Public/PropertyBindingExtension.h
--- a/Engine/Plugins/Runtime/PropertyBindingUtils/Source/PropertyBindingUtilsEditor/Public/PropertyBindingExtension.h
+++ b/Engine/Plugins/Runtime/PropertyBindingUtils/Source/PropertyBindingUtilsEditor/Public/PropertyBindingExtension.h
@@ -14,6 +14,7 @@ class IDetailLayoutBuilder;
 class IPropertyAccessEditor;
 class IPropertyBindingBindingCollectionOwner;
 class IPropertyHandle;
+class UEdGraphSchema_K2;
 struct FBindingChainElement;
 struct FBindingContextStruct;
 struct FPropertyBindingBindableStructDescriptor;
@@ -56,7 +57,9 @@ const FLazyName MetaDataNoPromoteToParameter("NoPromoteToParameter");
  */
 const FLazyName MetaDataBindingRoot("BindingRoot");
 
+UE_DEPRECATED(5.8, "Use the overload with Schema instead.")
 PROPERTYBINDINGUTILSEDITOR_API extern FText GetPropertyTypeText(const FProperty* Property);
+PROPERTYBINDINGUTILSEDITOR_API extern FText GetPropertyTypeText(TNotNull<const UEdGraphSchema_K2*> Schema, const FProperty* Property);
 
 /**
  * Returns property path for a specific property.
@@ -116,6 +119,7 @@ struct FCachedBindingData
 	[[nodiscard]] const FPropertyBindingPath& GetSourcePath() const;
 	[[nodiscard]] const FPropertyBindingPath& GetTargetPath() const;
 	[[nodiscard]] TConstArrayView<TInstancedStruct<FPropertyBindingBindableStructDescriptor>> GetAccessibleStructs() const;
+	[[nodiscard]] PROPERTYBINDINGUTILSEDITOR_API virtual TNotNull<const UEdGraphSchema_K2*> GetEdGraphSchema() const;
 
 	[[nodiscard]] FText GetFormatableText() const;
 	[[nodiscard]] FText GetFormatableTooltipText() const;
@@ -130,9 +134,19 @@ struct FCachedBindingData
 	PROPERTYBINDINGUTILSEDITOR_API virtual void AddPropertyInfoOverride(const FProperty& InProperty, TArray<TSharedPtr<const UE::PropertyBinding::FPropertyInfoOverride>>& OutPropertyInfoOverrides) const;
 	PROPERTYBINDINGUTILSEDITOR_API virtual void GetSourceDataViewForNewBinding(TNotNull<IPropertyBindingBindingCollectionOwner*> InBindingsOwner,TConstStructView<FPropertyBindingBindableStructDescriptor> InDescriptor, FPropertyBindingDataView& OutSourceDataView);
 
-	UE_DEPRECATED(5.7, "Use the overload with OutIconBrush instead")
+	struct FPinInfo
+	{
+		FEdGraphPinType Type;
+		const FSlateBrush* IconBrush = nullptr;
+		FLinearColor IconColor = FLinearColor::White;
+	};
+
+
+	UE_DEPRECATED(5.7, "Use GetPinInfoForProperty instead")
 	[[nodiscard]] PROPERTYBINDINGUTILSEDITOR_API virtual bool GetPinTypeAndIconForProperty(const FProperty& InProperty, FPropertyBindingDataView InTargetDataView, FEdGraphPinType& OutPinType, FName& OutIconName) const final;
-	[[nodiscard]] PROPERTYBINDINGUTILSEDITOR_API virtual bool GetPinTypeAndIconForProperty(const FProperty& InProperty, FPropertyBindingDataView InTargetDataView, FEdGraphPinType& OutPinType, const FSlateBrush*& OutIconBrush) const;
+	UE_DEPRECATED(5.8, "Use GetPinInfoForProperty instead")
+	[[nodiscard]] PROPERTYBINDINGUTILSEDITOR_API virtual bool GetPinTypeAndIconForProperty(const FProperty& InProperty, FPropertyBindingDataView InTargetDataView, FEdGraphPinType& OutPinType, const FSlateBrush*& OutIconBrush) const final;
+	[[nodiscard]] PROPERTYBINDINGUTILSEDITOR_API virtual FPinInfo GetPinInfoForProperty(TNotNull<const FProperty*> InProperty, FPropertyBindingDataView InTargetDataView) const;
 	[[nodiscard]] PROPERTYBINDINGUTILSEDITOR_API virtual bool GetPropertyFunctionText(FConstStructView InPropertyFunctionStructView, FText& OutText) const;
 	[[nodiscard]] PROPERTYBINDINGUTILSEDITOR_API virtual bool GetPropertyFunctionTooltipText(FConstStructView InPropertyFunctionStructView, FText& OutText) const;
 	[[nodiscard]] PROPERTYBINDINGUTILSEDITOR_API virtual bool GetPropertyFunctionIconColor(FConstStructView InPropertyFunctionStructView, FLinearColor& OutColor) const;
@@ -257,8 +271,11 @@ protected:
 	}
 
 	PROPERTYBINDINGUTILSEDITOR_API virtual bool CanBindToArrayElements() const;
+	[[nodiscard]] PROPERTYBINDINGUTILSEDITOR_API virtual TNotNull<const UEdGraphSchema_K2*> GetEdGraphSchema() const;
 
+	UE_DEPRECATED(5.8, "Use the overload with Schema instead.")
 	PROPERTYBINDINGUTILSEDITOR_API static TSharedRef<SWidget> MakeContextStructWidget(const FPropertyBindingBindableStructDescriptor& InContextStruct);
+	PROPERTYBINDINGUTILSEDITOR_API static TSharedRef<SWidget> MakeContextStructWidget(TNotNull<const UEdGraphSchema_K2*> Schema, const FPropertyBindingBindableStructDescriptor& InContextStruct);
 
 private:
 	bool CanPromoteToParameter(const TSharedPtr<IPropertyHandle>& InPropertyHandle) const;
diff --git a/Engine/Plugins/Runtime/StateTree/Source/StateTreeEditorModule/Private/Customizations/StateTreeBindingExtension.cpp b/Engine/Plugins/Runtime/StateTree/Source/StateTreeEditorModule/Private/Customizations/StateTreeBindingExtension.cpp
--- a/Engine/Plugins/Runtime/StateTree/Source/StateTreeEditorModule/Private/Customizations/StateTreeBindingExtension.cpp
+++ b/Engine/Plugins/Runtime/StateTree/Source/StateTreeEditorModule/Private/Customizations/StateTreeBindingExtension.cpp
@@ -4,13 +4,13 @@
 
 #include "ScopedTransaction.h"
 #include "DetailLayoutBuilder.h"
-#include "EdGraphSchema_K2.h"
 #include "Features/IModularFeatures.h"
 #include "IPropertyAccessEditor.h"
 #include "IPropertyUtilities.h"
 #include "StateTreeAnyEnum.h"
 #include "StateTreeCompiler.h"
 #include "StateTreeEditorPropertyBindings.h"
+#include "StateTreeGraphSchema.h"
 #include "StateTreeNodeBase.h"
 #include "UObject/EnumProperty.h"
 #include "Widgets/Images/SImage.h"
@@ -191,6 +191,11 @@ struct FStateTreeCachedBindingData : UE::PropertyBinding::FCachedBindingData, TS
 		}
 	}
 
+	virtual TNotNull<const UEdGraphSchema_K2*> GetEdGraphSchema() const override
+	{
+		return GetDefault<UE::StateTree::Editor::UEdGraphSchema_StateTree>();
+	}
+
 	virtual bool AddBindingInternal(TConstStructView<FPropertyBindingBindableStructDescriptor> InDescriptor,
 		FPropertyBindingPath& InOutSourcePath,
 		const FPropertyBindingPath& InTargetPath) override
@@ -289,30 +294,33 @@ struct FStateTreeCachedBindingData : UE::PropertyBinding::FCachedBindingData, TS
 		}
 	}
 
-	virtual bool GetPinTypeAndIconForProperty(const FProperty& InProperty, FPropertyBindingDataView InTargetDataView, FEdGraphPinType& OutPinType, const FSlateBrush*& OutIconBrush) const override
+	virtual FPinInfo GetPinInfoForProperty(TNotNull<const FProperty*> InProperty, FPropertyBindingDataView InTargetDataView) const override
 	{
-		const bool bIsPropertyRef = PropertyRefHelpers::IsPropertyRef(InProperty);;
+		const bool bIsPropertyRef = PropertyRefHelpers::IsPropertyRef(*InProperty);
 		if (bIsPropertyRef && InTargetDataView.IsValid())
 		{
+			FPinInfo Result = Super::GetPinInfoForProperty(InProperty, InTargetDataView);
 			// Use internal type to construct PinType if it's property of PropertyRef type.
 			TArray<FPropertyBindingPathIndirection> TargetIndirections;
 			if (ensure(GetTargetPath().ResolveIndirectionsWithValue(InTargetDataView, TargetIndirections)))
 			{
 				const void* PropertyRef = TargetIndirections.Last().GetPropertyAddress();
-				OutPinType = PropertyRefHelpers::GetPropertyRefInternalTypeAsPin(InProperty, PropertyRef);
+				Result.Type = PropertyRefHelpers::GetPropertyRefInternalTypeAsPin(*InProperty, PropertyRef);
 			}
-			OutIconBrush = FAppStyle::GetBrush("Kismet.Tabs.Variables");
-			return true;
+			return Result;
 		}
 
-		if (IsDelegateListenerProperty(&InProperty))
+		if (IsDelegateListenerProperty(InProperty))
 		{
-			OutPinType.PinCategory = UEdGraphSchema_K2::PC_Delegate;
-			OutIconBrush = FAppStyle::GetBrush("Icons.Event");
-			return true;
+			FPinInfo Result;
+			Result.Type.PinCategory = UEdGraphSchema_K2::PC_Struct;
+			Result.Type.PinSubCategoryObject = FStateTreeDelegateListener::StaticStruct();
+			Result.IconBrush = FAppStyle::GetBrush("Icons.Event");
+			Result.IconColor = FLinearColor::White;
+			return Result;
 		}
 
-		return false;
+		return Super::GetPinInfoForProperty(InProperty, InTargetDataView);
 	}
 
 	virtual bool IsPropertyReference(const FProperty& InProperty) override
@@ -787,6 +795,8 @@ bool FStateTreeBindingExtension::CanBindToProperty(const FPropertyBindingPath& I
 
 void FStateTreeBindingExtension::UpdateContextStruct(TConstStructView<FPropertyBindingBindableStructDescriptor> InStructDesc, FBindingContextStruct& InOutContextStruct, TMap<FString, FText>& InOutSectionNames) const
 {
+	using namespace UE::StateTree::Editor;
+
 	const FStateTreeBindableStructDesc& StructDesc = InStructDesc.Get<FStateTreeBindableStructDesc>();
 	// Mare sure same section names get exact same FText representation (binding widget uses IsIdentical() to compare the section names).
 	if (const FText* SectionText = InOutSectionNames.Find(StructDesc.StatePath))
@@ -803,7 +813,7 @@ void FStateTreeBindingExtension::UpdateContextStruct(TConstStructView<FPropertyB
 	{
 		if (const FProperty* OutputProperty = UE::StateTree::GetStructSingleOutputProperty(*StructDesc.Struct))
 		{
-			const UEdGraphSchema_K2* Schema = GetDefault<UEdGraphSchema_K2>();
+			const UEdGraphSchema_StateTree* Schema = GetDefault<UEdGraphSchema_StateTree>();
 			check(Schema);
 
 			FEdGraphPinType PinType;
@@ -851,6 +861,11 @@ void FStateTreeBindingsChildrenCustomization::CustomizeChildren(IDetailChildrenB
 	}
 }
 
+TNotNull<const UEdGraphSchema_K2*> FStateTreeBindingExtension::GetEdGraphSchema() const
+{
+	return GetDefault<UE::StateTree::Editor::UEdGraphSchema_StateTree>();
+}
+
 void FStateTreeBindingExtension::CustomizeDetailWidgetRow(FDetailWidgetRow& InWidgetRow, const IDetailLayoutBuilder& InDetailBuilder, IPropertyBindingBindingCollectionOwner* InBindingsOwner,
 	TSharedPtr<IPropertyHandle> InPropertyHandle, const FPropertyBindingPath& InTargetPath,
 	TSharedPtr<UE::PropertyBinding::FCachedBindingData> InCachedBindingData) const
diff --git a/Engine/Plugins/Runtime/StateTree/Source/StateTreeEditorModule/Private/Customizations/StateTreeBindingExtension.h b/Engine/Plugins/Runtime/StateTree/Source/StateTreeEditorModule/Private/Customizations/StateTreeBindingExtension.h
--- a/Engine/Plugins/Runtime/StateTree/Source/StateTreeEditorModule/Private/Customizations/StateTreeBindingExtension.h
+++ b/Engine/Plugins/Runtime/StateTree/Source/StateTreeEditorModule/Private/Customizations/StateTreeBindingExtension.h
@@ -31,7 +31,7 @@ private:
 	{
 		return true;
 	}
-
+	virtual TNotNull<const UEdGraphSchema_K2*> GetEdGraphSchema() const override;
 	virtual void CustomizeDetailWidgetRow(FDetailWidgetRow& InWidgetRow, const IDetailLayoutBuilder& InDetailBuilder, IPropertyBindingBindingCollectionOwner* InBindingsOwner, TSharedPtr<IPropertyHandle> InPropertyHandle, const FPropertyBindingPath& InTargetPath, TSharedPtr<UE::PropertyBinding::FCachedBindingData> InCachedBindingData) const override;
 };
 
diff --git a/Engine/Plugins/Runtime/StateTree/Source/StateTreeEditorModule/Private/Customizations/StateTreeBlueprintPropertyRefDetails.cpp b/Engine/Plugins/Runtime/StateTree/Source/StateTreeEditorModule/Private/Customizations/StateTreeBlueprintPropertyRefDetails.cpp
--- a/Engine/Plugins/Runtime/StateTree/Source/StateTreeEditorModule/Private/Customizations/StateTreeBlueprintPropertyRefDetails.cpp
+++ b/Engine/Plugins/Runtime/StateTree/Source/StateTreeEditorModule/Private/Customizations/StateTreeBlueprintPropertyRefDetails.cpp
@@ -3,8 +3,9 @@
 #include "StateTreeBlueprintPropertyRefDetails.h"
 #include "DetailWidgetRow.h"
 #include "SPinTypeSelector.h"
-#include "StateTreePropertyRef.h"
 #include "StateTreeEditorNode.h"
+#include "StateTreeGraphSchema.h"
+#include "StateTreePropertyRef.h"
 #include "IDetailChildrenBuilder.h"
 #include "DetailLayoutBuilder.h"
 
@@ -128,7 +129,7 @@ void FStateTreeBlueprintPropertyRefDetails::CustomizeHeader(TSharedRef<IProperty
 		+ SHorizontalBox::Slot()
 		.AutoWidth()
 		[
-			SNew(SPinTypeSelector, FGetPinTypeTree::CreateUObject(GetDefault<UStateTreePropertyRefSchema>(), &UEdGraphSchema_K2::GetVariableTypeTree))
+			SNew(SPinTypeSelector, FGetPinTypeTree::CreateUObject(GetDefault<UStateTreePropertyRefSchema>(), &UE::StateTree::Editor::UEdGraphSchema_StateTree::GetVariableTypeTree))
 				.OnPinTypeChanged_Lambda(PinInfoChanged)
 				.TargetPinType_Lambda(GetPinInfo)
 				.Schema(GetDefault<UStateTreePropertyRefSchema>())
diff --git a/Engine/Plugins/Runtime/StateTree/Source/StateTreeEditorModule/Private/Customizations/StateTreeEditorDataDetails.cpp b/Engine/Plugins/Runtime/StateTree/Source/StateTreeEditorModule/Private/Customizations/StateTreeEditorDataDetails.cpp
--- a/Engine/Plugins/Runtime/StateTree/Source/StateTreeEditorModule/Private/Customizations/StateTreeEditorDataDetails.cpp
+++ b/Engine/Plugins/Runtime/StateTree/Source/StateTreeEditorModule/Private/Customizations/StateTreeEditorDataDetails.cpp
@@ -12,6 +12,7 @@
 #include "StateTreeEditorNodeUtils.h"
 #include "StateTreeEditorSchema.h"
 #include "StateTreeEditorStyle.h"
+#include "StateTreeGraphSchema.h"
 #include "StateTreeSchema.h"
 #include "Styling/StyleColors.h"
 #include "Widgets/Images/SImage.h"
@@ -27,6 +28,8 @@ TSharedRef<IDetailCustomization> FStateTreeEditorDataDetails::MakeInstance()
 
 void FStateTreeEditorDataDetails::CustomizeDetails(IDetailLayoutBuilder& DetailBuilder)
 {
+	using namespace UE::StateTree::Editor;
+
 	const TSharedRef<IPropertyUtilities> PropUtils = DetailBuilder.GetPropertyUtilities();
 	WeakPropertyUtilities = PropUtils.ToWeakPtr();
 
@@ -82,7 +85,7 @@ void FStateTreeEditorDataDetails::CustomizeDetails(IDetailLayoutBuilder& DetailB
 			}
 			PinType.PinSubCategoryObject = const_cast<UStruct*>(ContextData.Struct.Get());
 
-			const UEdGraphSchema_K2* EdGraphSchema = GetDefault<UEdGraphSchema_K2>();
+			const UEdGraphSchema_StateTree* EdGraphSchema = GetDefault<UEdGraphSchema_StateTree>();
 			const FSlateBrush* Icon = FBlueprintEditorUtils::GetIconFromPin(PinType, true);
 			const FLinearColor Color = EdGraphSchema->GetPinTypeColor(PinType);
 
diff --git a/Engine/Plugins/Runtime/StateTree/Source/StateTreeEditorModule/Private/Customizations/StateTreeEditorNodeDetails.cpp b/Engine/Plugins/Runtime/StateTree/Source/StateTreeEditorModule/Private/Customizations/StateTreeEditorNodeDetails.cpp
--- a/Engine/Plugins/Runtime/StateTree/Source/StateTreeEditorModule/Private/Customizations/StateTreeEditorNodeDetails.cpp
+++ b/Engine/Plugins/Runtime/StateTree/Source/StateTreeEditorModule/Private/Customizations/StateTreeEditorNodeDetails.cpp
@@ -43,6 +43,7 @@
 #include "SourceCodeNavigation.h"
 #include "StateTreeDelegate.h"
 #include "StateTreeEditorDataClipboardHelpers.h"
+#include "StateTreeGraphSchema.h"
 #include "Widgets/Notifications/SNotificationList.h"
 
 #define LOCTEXT_NAMESPACE "StateTreeEditor"
@@ -77,6 +78,8 @@ namespace UE::StateTreeEditor::Internal
 	/** @return text describing the pin type, matches SPinTypeSelector. */
 	FText GetPinTypeText(const FEdGraphPinType& PinType)
 	{
+		using namespace UE::StateTree::Editor;
+
 		const FName PinSubCategory = PinType.PinSubCategory;
 		const UObject* PinSubCategoryObject = PinType.PinSubCategoryObject.Get();
 		if (PinSubCategory != UEdGraphSchema_K2::PSC_Bitmask && PinSubCategoryObject)
@@ -88,7 +91,7 @@ namespace UE::StateTreeEditor::Internal
 			return FText::FromString(PinSubCategoryObject->GetName());
 		}
 
-		return UEdGraphSchema_K2::GetCategoryText(PinType.PinCategory, NAME_None, true);
+		return UEdGraphSchema_StateTree::GetCategoryText(PinType.PinCategory, NAME_None, true);
 	}
 
 	/** @return if property is struct property of DelegateDispatcher type. */
@@ -125,6 +128,8 @@ namespace UE::StateTreeEditor::Internal
 
 	void ModifyRow(IDetailPropertyRow& ChildRow, const FGuid& ID, UStateTreeEditorData* EditorData)
 	{
+		using namespace UE::StateTree::Editor;
+
 		FStateTreeEditorPropertyBindings* EditorPropBindings = EditorData ? EditorData->GetPropertyEditorBindings() : nullptr;
 		if (!EditorPropBindings)
 		{
@@ -162,7 +167,7 @@ namespace UE::StateTreeEditor::Internal
 			if (bValidUsage || bIsDelegateDispatcher)
 			{
 				FEdGraphPinType PinType;
-				const UEdGraphSchema_K2* Schema = GetDefault<UEdGraphSchema_K2>();
+				const UEdGraphSchema_StateTree* Schema = GetDefault<UEdGraphSchema_StateTree>();
 
 				// Show referenced type for property refs.
 				if (UE::StateTree::PropertyRefHelpers::IsPropertyRef(*Property))
@@ -204,7 +209,8 @@ namespace UE::StateTreeEditor::Internal
 					LabelToolTip = LOCTEXT("DelegateToolTip", "This is Delegate Dispatcher. You can bind to it from listeners.");
 
 					FEdGraphPinType DelegatePinType;
-					DelegatePinType.PinCategory = UEdGraphSchema_K2::PC_Delegate;
+					DelegatePinType.PinCategory = UEdGraphSchema_K2::PC_Struct;
+					DelegatePinType.PinSubCategoryObject = FStateTreeDelegateDispatcher::StaticStruct();
 					IconColor = Schema->GetPinTypeColor(DelegatePinType);
 				}
 				else if (Usage == EStateTreePropertyUsage::Input)
diff --git a/Engine/Plugins/Runtime/StateTree/Source/StateTreeEditorModule/Private/Customizations/StateTreeEventDescDetails.cpp b/Engine/Plugins/Runtime/StateTree/Source/StateTreeEditorModule/Private/Customizations/StateTreeEventDescDetails.cpp
--- a/Engine/Plugins/Runtime/StateTree/Source/StateTreeEditorModule/Private/Customizations/StateTreeEventDescDetails.cpp
+++ b/Engine/Plugins/Runtime/StateTree/Source/StateTreeEditorModule/Private/Customizations/StateTreeEventDescDetails.cpp
@@ -10,6 +10,7 @@
 #include "PropertyBindingExtension.h"
 #include "StateTreeBindingExtension.h"
 #include "StateTreeEditorStyle.h"
+#include "StateTreeGraphSchema.h"
 #include "Styling/StyleColors.h"
 #include "Kismet2/BlueprintEditorUtils.h"
 #include "Widgets/Images/SImage.h"
@@ -140,7 +141,9 @@ private:
 
 	TSharedRef<SWidget> GenerateValueWidget(const FProperty& Property) const
 	{
-		const UEdGraphSchema_K2* Schema = GetDefault<UEdGraphSchema_K2>();
+		using namespace UE::StateTree::Editor;
+
+		const UEdGraphSchema_StateTree* Schema = GetDefault<UEdGraphSchema_StateTree>();
 		FEdGraphPinType PinType;
 		Schema->ConvertPropertyToPinType(&Property, PinType);
 
@@ -171,6 +174,8 @@ private:
 
 	static FText GetPinTypeText(const FEdGraphPinType& PinType)
 	{
+		using namespace UE::StateTree::Editor;
+
 		const FName PinSubCategory = PinType.PinSubCategory;
 		const UObject* PinSubCategoryObject = PinType.PinSubCategoryObject.Get();
 		if (PinSubCategory != UEdGraphSchema_K2::PSC_Bitmask && PinSubCategoryObject)
@@ -182,7 +187,7 @@ private:
 			return FText::FromString(PinSubCategoryObject->GetName());
 		}
 
-		return UEdGraphSchema_K2::GetCategoryText(PinType.PinCategory, NAME_None, true);
+		return UEdGraphSchema_StateTree::GetCategoryText(PinType.PinCategory, NAME_None, true);
 	}
 
 	FName Name;
diff --git a/Engine/Plugins/Runtime/StateTree/Source/StateTreeEditorModule/Private/K2Node_StateTreeBlueprintPropertyRef.cpp b/Engine/Plugins/Runtime/StateTree/Source/StateTreeEditorModule/Private/K2Node_StateTreeBlueprintPropertyRef.cpp
--- a/Engine/Plugins/Runtime/StateTree/Source/StateTreeEditorModule/Private/K2Node_StateTreeBlueprintPropertyRef.cpp
+++ b/Engine/Plugins/Runtime/StateTree/Source/StateTreeEditorModule/Private/K2Node_StateTreeBlueprintPropertyRef.cpp
@@ -3,14 +3,14 @@
 #include "K2Node_StateTreeBlueprintPropertyRef.h"
 #include "StateTreePropertyRef.h"
 #include "Blueprint/StateTreeNodeBlueprintBase.h"
-#include "EdGraphSchema_K2.h"
 #include "BlueprintActionDatabaseRegistrar.h"
 #include "BlueprintNodeSpawner.h"
+#include "EdGraphUtilities.h"
 #include "Kismet2/BlueprintEditorUtils.h"
 #include "K2Node_VariableGet.h"
+#include "StateTreeGraphSchema.h"
 #include "StateTreePropertyRefHelpers.h"
 #include "KismetCompiledFunctionContext.h"
-#include "EdGraphUtilities.h"
 
 #include UE_INLINE_GENERATED_CPP_BY_NAME(K2Node_StateTreeBlueprintPropertyRef)
 
@@ -112,7 +112,7 @@ void UK2Node_StateTreeBlueprintPropertyRef::UpdateOutputPin() const
 
 				PinType.bIsReference = true;
 				OutValuePin.PinType = PinType;
-				GetDefault<UEdGraphSchema_K2>()->SetPinAutogeneratedDefaultValueBasedOnType(&OutValuePin);
+				GetDefault<UE::StateTree::Editor::UEdGraphSchema_StateTree>()->SetPinAutogeneratedDefaultValueBasedOnType(&OutValuePin);
 				return;
 			}
 		}
diff --git a/Engine/Plugins/Runtime/StateTree/Source/StateTreeEditorModule/Private/K2Node_StateTreeNodeGetPropertyDescription.cpp b/Engine/Plugins/Runtime/StateTree/Source/StateTreeEditorModule/Private/K2Node_StateTreeNodeGetPropertyDescription.cpp
--- a/Engine/Plugins/Runtime/StateTree/Source/StateTreeEditorModule/Private/K2Node_StateTreeNodeGetPropertyDescription.cpp
+++ b/Engine/Plugins/Runtime/StateTree/Source/StateTreeEditorModule/Private/K2Node_StateTreeNodeGetPropertyDescription.cpp
@@ -2,11 +2,11 @@
 
 #include "K2Node_StateTreeNodeGetPropertyDescription.h"
 #include "Blueprint/StateTreeNodeBlueprintBase.h"
-#include "EdGraphSchema_K2.h"
 #include "BlueprintActionDatabaseRegistrar.h"
 #include "BlueprintNodeSpawner.h"
 #include "K2Node_CallFunction.h"
 #include "KismetCompiler.h"
+#include "StateTreeGraphSchema.h"
 
 #include UE_INLINE_GENERATED_CPP_BY_NAME(K2Node_StateTreeNodeGetPropertyDescription)
 
diff --git a/Engine/Plugins/Runtime/StateTree/Source/StateTreeEditorModule/Private/K2Node_StateTreeReference.cpp b/Engine/Plugins/Runtime/StateTree/Source/StateTreeEditorModule/Private/K2Node_StateTreeReference.cpp
--- a/Engine/Plugins/Runtime/StateTree/Source/StateTreeEditorModule/Private/K2Node_StateTreeReference.cpp
+++ b/Engine/Plugins/Runtime/StateTree/Source/StateTreeEditorModule/Private/K2Node_StateTreeReference.cpp
@@ -6,7 +6,6 @@
 #include "BlueprintNodeSpawner.h"
 #include "EditorCategoryUtils.h"
 #include "EdGraph/EdGraphPin.h"
-#include "EdGraphSchema_K2.h"
 #include "K2Node_Variable.h"
 #include "KismetCompiler.h"
 #include "KismetCompilerMisc.h"
@@ -14,6 +13,7 @@
 #include "StateTree.h"
 #include "StateTreeDelegates.h"
 #include "StateTreeFunctionLibrary.h"
+#include "StateTreeGraphSchema.h"
 #include "StateTreeReference.h"
 #include "StructUtils/PropertyBag.h"
 
@@ -52,6 +52,8 @@ namespace UE::StateTreeEditor::Private
 
 	bool DoRenamedPinsMatch(UStateTree* StateTree, const UEdGraphPin* NewPin, const UEdGraphPin* OldPin)
 	{
+		using namespace UE::StateTree::Editor;
+
 		if (StateTree == nullptr
 			|| NewPin == nullptr
 			|| OldPin == nullptr
@@ -66,7 +68,7 @@ namespace UE::StateTreeEditor::Private
 			return false;
 		}
 
-		const bool bCompatible = GetDefault<UEdGraphSchema_K2>()->ArePinTypesCompatible(NewPin->PinType, OldPin->PinType);
+		const bool bCompatible = GetDefault<UEdGraphSchema_StateTree>()->ArePinTypesCompatible(NewPin->PinType, OldPin->PinType);
 		if (!bCompatible)
 		{
 			return false;
@@ -131,11 +133,10 @@ FLinearColor UK2Node_MakeStateTreeReference::GetNodeTitleColor() const
 {
 	auto MakeTitleColor = []() -> FLinearColor
 		{
-			const UEdGraphSchema_K2* K2Schema = GetDefault<UEdGraphSchema_K2>();
 			FEdGraphPinType PinType;
 			PinType.PinCategory = UEdGraphSchema_K2::PC_Struct;
 			PinType.PinSubCategoryObject = FStateTreeReference::StaticStruct();
-			return K2Schema->GetPinTypeColor(PinType);
+			return GetDefault<UE::StateTree::Editor::UEdGraphSchema_StateTree>()->GetPinTypeColor(PinType);
 		};
 
 	static FLinearColor LocalCachedTitleColor = MakeTitleColor();
diff --git a/Engine/Plugins/Runtime/StateTree/Source/StateTreeEditorModule/Private/StateTreeGraphSchema.cpp b/Engine/Plugins/Runtime/StateTree/Source/StateTreeEditorModule/Private/StateTreeGraphSchema.cpp
new file mode 100644
--- /dev/null
+++ b/Engine/Plugins/Runtime/StateTree/Source/StateTreeEditorModule/Private/StateTreeGraphSchema.cpp
@@ -0,0 +1,7 @@
+// Copyright Epic Games, Inc. All Rights Reserved.
+
+#include "StateTreeGraphSchema.h"
+
+#include "StateTreeDelegate.h"
+
+#include UE_INLINE_GENERATED_CPP_BY_NAME(StateTreeGraphSchema)
diff --git a/Engine/Plugins/Runtime/StateTree/Source/StateTreeEditorModule/Private/StateTreeGraphSchema.h b/Engine/Plugins/Runtime/StateTree/Source/StateTreeEditorModule/Private/StateTreeGraphSchema.h
new file mode 100644
--- /dev/null
+++ b/Engine/Plugins/Runtime/StateTree/Source/StateTreeEditorModule/Private/StateTreeGraphSchema.h
@@ -0,0 +1,16 @@
+// Copyright Epic Games, Inc. All Rights Reserved.
+
+#pragma once
+
+#include "EdGraphSchema_K2.h"
+
+#include "StateTreeGraphSchema.generated.h"
+
+namespace UE::StateTree::Editor
+{
+	UCLASS()
+	class UEdGraphSchema_StateTree : public UEdGraphSchema_K2
+	{
+		GENERATED_BODY()
+	};
+}
diff --git a/Engine/Plugins/VirtualProduction/SceneState/Source/SceneStateBlueprintEditor/Private/DetailsView/SceneStateCachedBindingData.cpp b/Engine/Plugins/VirtualProduction/SceneState/Source/SceneStateBlueprintEditor/Private/DetailsView/SceneStateCachedBindingData.cpp
--- a/Engine/Plugins/VirtualProduction/SceneState/Source/SceneStateBlueprintEditor/Private/DetailsView/SceneStateCachedBindingData.cpp
+++ b/Engine/Plugins/VirtualProduction/SceneState/Source/SceneStateBlueprintEditor/Private/DetailsView/SceneStateCachedBindingData.cpp
@@ -78,21 +78,20 @@ void FSceneStateCachedBindingData::UpdatePropertyReferenceTooltip(const FPropert
 		}
 	}
 }
-
-bool FSceneStateCachedBindingData::GetPinTypeAndIconForProperty(const FProperty& InProperty, FPropertyBindingDataView InTargetDataView, FEdGraphPinType& OutPinType, const FSlateBrush*& OutIconBrush) const
+FSceneStateCachedBindingData::FPinInfo FSceneStateCachedBindingData::GetPinInfoForProperty(TNotNull<const FProperty*> InProperty, FPropertyBindingDataView InTargetDataView) const
 {
-	if (UE::SceneState::IsPropertyReference(&InProperty) && InTargetDataView.IsValid())
+	FPinInfo Result = Super::GetPinInfoForProperty(InProperty, InTargetDataView);
+	if (UE::SceneState::IsPropertyReference(InProperty) && InTargetDataView.IsValid())
 	{
 		// Use internal type to construct PinType if it's property of PropertyRef type.
 		TArray<FPropertyBindingPathIndirection> TargetIndirections;
 		if (ensure(GetTargetPath().ResolveIndirectionsWithValue(InTargetDataView, TargetIndirections)))
 		{
-			OutPinType = GetPropertyReferencePinType(&InProperty, TargetIndirections.Last().GetPropertyAddress());
+			Result.Type = GetPropertyReferencePinType(InProperty, TargetIndirections.Last().GetPropertyAddress());
 		}
-		OutIconBrush = FAppStyle::GetBrush("Kismet.Tabs.Variables");
-		return true;
+		Result.IconBrush = FAppStyle::GetBrush("Kismet.Tabs.Variables");
 	}
-	return false;
+	return Result;
 }
 
 bool FSceneStateCachedBindingData::IsPropertyReference(const FProperty& InProperty)
diff --git a/Engine/Plugins/VirtualProduction/SceneState/Source/SceneStateBlueprintEditor/Private/DetailsView/SceneStateCachedBindingData.h b/Engine/Plugins/VirtualProduction/SceneState/Source/SceneStateBlueprintEditor/Private/DetailsView/SceneStateCachedBindingData.h
--- a/Engine/Plugins/VirtualProduction/SceneState/Source/SceneStateBlueprintEditor/Private/DetailsView/SceneStateCachedBindingData.h
+++ b/Engine/Plugins/VirtualProduction/SceneState/Source/SceneStateBlueprintEditor/Private/DetailsView/SceneStateCachedBindingData.h
@@ -29,7 +29,7 @@ struct FSceneStateCachedBindingData : UE::PropertyBinding::FCachedBindingData, T
 
 	//~ Begin UE::PropertyBinding::FCachedBindingData
 	virtual void UpdatePropertyReferenceTooltip(const FProperty& InProperty, FTextBuilder& InOutTextBuilder) const override;
-	virtual bool GetPinTypeAndIconForProperty(const FProperty& InProperty, FPropertyBindingDataView InTargetDataView, FEdGraphPinType& OutPinType, const FSlateBrush*& OutIconBrush) const override;
+	virtual FPinInfo GetPinInfoForProperty(TNotNull<const FProperty*> InProperty, FPropertyBindingDataView InTargetDataView) const override;
 	virtual bool IsPropertyReference(const FProperty& InProperty) override;
 	virtual void UpdateSourcePropertyPath(TConstStructView<FPropertyBindingBindableStructDescriptor> InDescriptor, const FPropertyBindingPath& InSourcePath, FString& OutString);
 	virtual bool CanBindToContextStructInternal(const UStruct* InStruct, const int32 InStructIndex) override;
