[PrpertyBinding] Add config to keep some metadata on property promotion.

#Synced-CL 48343893

diff --git a/Engine/Plugins/Runtime/PropertyBindingUtils/Source/PropertyBindingUtils/Private/PropertyBindingBindingCollectionOwner.cpp b/Engine/Plugins/Runtime/PropertyBindingUtils/Source/PropertyBindingUtils/Private/PropertyBindingBindingCollectionOwner.cpp
new file mode 100644
--- /dev/null
+++ b/Engine/Plugins/Runtime/PropertyBindingUtils/Source/PropertyBindingUtils/Private/PropertyBindingBindingCollectionOwner.cpp
@@ -0,0 +1,23 @@
+// Copyright Epic Games, Inc. All Rights Reserved.
+
+#include "PropertyBindingBindingCollectionOwner.h"
+#include "PropertyBindingSettings.h"
+
+#include UE_INLINE_GENERATED_CPP_BY_NAME(PropertyBindingBindingCollectionOwner)
+
+#if WITH_EDITOR
+void IPropertyBindingBindingCollectionOwner::OnPromotingToParameter(FPropertyBagPropertyDesc& PropertyDesc)
+{
+	// By default, we create a descriptor based on the Target Property, but without the meta-data.
+	// This functionality mirrors the user action of adding a new property from the UI, where meta-data is not available.
+	// Additionally, meta-data like EditCondition is not desirable here.
+	// UPropertyBindingSettings::MetaDataToKeepWhenPromotingToParameter can be modified if a specific metadata needs to stay.
+	// Derived classes can override this behavior.
+	PropertyDesc.MetaClass = nullptr;
+	PropertyDesc.MetaData.RemoveAllSwap(
+		[](const FPropertyBagPropertyDescMetaData& MetaData)
+		{
+			return !GetDefault<UPropertyBindingSettings>()->MetaDataToKeepWhenPromotingToParameter.Contains(MetaData.Key);
+		});
+}
+#endif
diff --git a/Engine/Plugins/Runtime/PropertyBindingUtils/Source/PropertyBindingUtils/Private/PropertyBindingSettings.cpp b/Engine/Plugins/Runtime/PropertyBindingUtils/Source/PropertyBindingUtils/Private/PropertyBindingSettings.cpp
new file mode 100644
--- /dev/null
+++ b/Engine/Plugins/Runtime/PropertyBindingUtils/Source/PropertyBindingUtils/Private/PropertyBindingSettings.cpp
@@ -0,0 +1,13 @@
+// Copyright Epic Games, Inc. All Rights Reserved.
+
+#include "PropertyBindingSettings.h"
+
+#include UE_INLINE_GENERATED_CPP_BY_NAME(PropertyBindingSettings)
+
+UPropertyBindingSettings::UPropertyBindingSettings()
+{
+	//These metadata are needed to keep the display.
+	//Additional metadata can be added to the ini file.
+	MetaDataToKeepWhenPromotingToParameter.Add("Bitmask");
+	MetaDataToKeepWhenPromotingToParameter.Add("BitmaskEnum");
+}
diff --git a/Engine/Plugins/Runtime/PropertyBindingUtils/Source/PropertyBindingUtils/Private/PropertyBindingSettings.h b/Engine/Plugins/Runtime/PropertyBindingUtils/Source/PropertyBindingUtils/Private/PropertyBindingSettings.h
new file mode 100644
--- /dev/null
+++ b/Engine/Plugins/Runtime/PropertyBindingUtils/Source/PropertyBindingUtils/Private/PropertyBindingSettings.h
@@ -0,0 +1,21 @@
+// Copyright Epic Games, Inc. All Rights Reserved.
+
+#pragma once
+
+#include "Engine/DeveloperSettings.h"
+#include "PropertyBindingSettings.generated.h"
+
+UCLASS(MinimalAPI, config = Engine, defaultconfig)
+class UPropertyBindingSettings : public UDeveloperSettings
+{
+	GENERATED_BODY()
+
+public:
+	UPropertyBindingSettings();
+
+	/**
+	 * Meta-data that are not in the MetaDataToKeepWhenPromotingToParameter are removed when a property if promoted.
+	 */
+	UPROPERTY(EditAnywhere, config, Category = "Property Binding")
+	TArray<FName> MetaDataToKeepWhenPromotingToParameter;
+};
\ No newline at end of file
diff --git a/Engine/Plugins/Runtime/PropertyBindingUtils/Source/PropertyBindingUtils/PropertyBindingUtils.Build.cs b/Engine/Plugins/Runtime/PropertyBindingUtils/Source/PropertyBindingUtils/PropertyBindingUtils.Build.cs
--- a/Engine/Plugins/Runtime/PropertyBindingUtils/Source/PropertyBindingUtils/PropertyBindingUtils.Build.cs
+++ b/Engine/Plugins/Runtime/PropertyBindingUtils/Source/PropertyBindingUtils/PropertyBindingUtils.Build.cs
@@ -16,6 +16,7 @@ namespace UnrealBuildTool.Rules
 
 			PrivateDependencyModuleNames.AddRange(
 			new string[] {
+				"DeveloperSettings",
 				"RenderCore",
 				"InputCore"
 			}
diff --git a/Engine/Plugins/Runtime/PropertyBindingUtils/Source/PropertyBindingUtils/Public/PropertyBindingBindingCollectionOwner.h b/Engine/Plugins/Runtime/PropertyBindingUtils/Source/PropertyBindingUtils/Public/PropertyBindingBindingCollectionOwner.h
--- a/Engine/Plugins/Runtime/PropertyBindingUtils/Source/PropertyBindingUtils/Public/PropertyBindingBindingCollectionOwner.h
+++ b/Engine/Plugins/Runtime/PropertyBindingUtils/Source/PropertyBindingUtils/Public/PropertyBindingBindingCollectionOwner.h
@@ -6,6 +6,8 @@
 #include "UObject/Interface.h"
 #include "PropertyBindingBindingCollectionOwner.generated.h"
 
+#define UE_API PROPERTYBINDINGUTILS_API
+
 namespace UE::PropertyBinding
 {
 struct FPropertyCreationDescriptor;
@@ -17,7 +19,6 @@ struct FPropertyBindingDataView;
 struct FPropertyBindingPath;
 template<typename BaseStructT> struct TInstancedStruct;
 
-
 UINTERFACE(MinimalAPI)
 class UPropertyBindingBindingCollectionOwner : public UInterface
 {
@@ -98,15 +99,7 @@ public:
 	 * Called when a property is being promoted to a parameter.
 	 * @param PropertyDesc The newly created property descriptor
 	 */
-	virtual void OnPromotingToParameter(FPropertyBagPropertyDesc& PropertyDesc)
-	{
-		// By default, we create a descriptor based on the Target Property, but without the meta-data.
-		// This functionality mirrors the user action of adding a new property from the UI, where meta-data is not available.
-		// Additionally, meta-data like EditCondition is not desirable here.
-		// Derived classes can override this behavior.
-		PropertyDesc.MetaClass = nullptr;
-		PropertyDesc.MetaData.Reset();
-	}
+	UE_API virtual void OnPromotingToParameter(FPropertyBagPropertyDesc& PropertyDesc);
 
 	/**
 	 * Creates the given properties in the property bag of the struct matching the given struct ID
@@ -133,3 +126,5 @@ public:
 	}
 #endif // WITH_EDITOR
 };
+
+#undef UE_API
\ No newline at end of file
