Fix typo in a comment.
#rnx
#Synced-CL 48343893

diff --git a/Engine/Plugins/Runtime/PropertyBindingUtils/Source/PropertyBindingUtils/Private/PropertyBindingSettings.h b/Engine/Plugins/Runtime/PropertyBindingUtils/Source/PropertyBindingUtils/Private/PropertyBindingSettings.h
--- a/Engine/Plugins/Runtime/PropertyBindingUtils/Source/PropertyBindingUtils/Private/PropertyBindingSettings.h
+++ b/Engine/Plugins/Runtime/PropertyBindingUtils/Source/PropertyBindingUtils/Private/PropertyBindingSettings.h
@@ -14,7 +14,7 @@ public:
 	UPropertyBindingSettings();
 
 	/**
-	 * Meta-data that are not in the MetaDataToKeepWhenPromotingToParameter are removed when a property if promoted.
+	 * Meta-data that are not in the MetaDataToKeepWhenPromotingToParameter are removed when a property is promoted.
 	 */
 	UPROPERTY(EditAnywhere, config, Category = "Property Binding")
 	TArray<FName> MetaDataToKeepWhenPromotingToParameter;
