[PropertyBindingUtils] Fixed incorrect loctext namespace

#Synced-CL 48638146

diff --git a/Engine/Plugins/Runtime/PropertyBindingUtils/Source/PropertyBindingUtilsEditor/Private/PropertyBindingExtension.cpp b/Engine/Plugins/Runtime/PropertyBindingUtils/Source/PropertyBindingUtilsEditor/Private/PropertyBindingExtension.cpp
--- a/Engine/Plugins/Runtime/PropertyBindingUtils/Source/PropertyBindingUtilsEditor/Private/PropertyBindingExtension.cpp
+++ b/Engine/Plugins/Runtime/PropertyBindingUtils/Source/PropertyBindingUtilsEditor/Private/PropertyBindingExtension.cpp
@@ -26,7 +26,7 @@
 #include "Widgets/Images/SImage.h"
 #include "Widgets/Layout/SBox.h"
 
-#define LOCTEXT_NAMESPACE "SmartObjectEditor"
+#define LOCTEXT_NAMESPACE "PropertyBindingExtension"
 
 struct FStateTreeBlueprintPropertyRef;
 struct FStateTreeEditorNode;
