[PropertyBindingUtils][State Tree]
Added option to not clear already resolved property paths if it failed so that caller can have partial resolved paths.

#Synced-CL 48820855

diff --git a/Engine/Plugins/Runtime/PropertyBindingUtils/Source/PropertyBindingUtils/Private/PropertyBindingPath.cpp b/Engine/Plugins/Runtime/PropertyBindingUtils/Source/PropertyBindingUtils/Private/PropertyBindingPath.cpp
--- a/Engine/Plugins/Runtime/PropertyBindingUtils/Source/PropertyBindingUtils/Private/PropertyBindingPath.cpp
+++ b/Engine/Plugins/Runtime/PropertyBindingUtils/Source/PropertyBindingUtils/Private/PropertyBindingPath.cpp
@@ -227,19 +227,19 @@ bool FPropertyBindingPath::Includes(const FPropertyBindingPath& Other) const
 	return true;
 }
 
-bool FPropertyBindingPath::ResolveIndirections(const UStruct* BaseStruct, TArray<FPropertyBindingPathIndirection>& OutIndirections, FString* OutError, bool bHandleRedirects) const
+bool FPropertyBindingPath::ResolveIndirections(const UStruct* BaseStruct, TArray<FPropertyBindingPathIndirection>& OutIndirections, FString* OutError, bool bHandleRedirects, bool bResetIndirectionsIfFailed) const
 {
-	return ResolveIndirectionsWithValue(FPropertyBindingDataView(BaseStruct, nullptr), OutIndirections, OutError, bHandleRedirects);
+	return ResolveIndirectionsWithValue(FPropertyBindingDataView(BaseStruct, nullptr), OutIndirections, OutError, bHandleRedirects, bResetIndirectionsIfFailed);
 }
 
 bool FPropertyBindingPath::ResolveIndirections(const UStruct* BaseStruct,
-	TArray<FPropertyBindingPathIndirection, TInlineAllocator<16>>& OutIndirections, FString* OutError, bool bHandleRedirects) const
+	TArray<FPropertyBindingPathIndirection, TInlineAllocator<16>>& OutIndirections, FString* OutError, bool bHandleRedirects, bool bResetIndirectionsIfFailed) const
 {
-	return ResolveIndirectionsWithValue(FPropertyBindingDataView(BaseStruct, nullptr), OutIndirections, OutError, bHandleRedirects);
+	return ResolveIndirectionsWithValue(FPropertyBindingDataView(BaseStruct, nullptr), OutIndirections, OutError, bHandleRedirects, bResetIndirectionsIfFailed);
 }
 
 template<typename Allocator>
-bool FPropertyBindingPath::ResolveIndirectionsWithValue(const FPropertyBindingDataView BaseValueView, TArray<FPropertyBindingPathIndirection, Allocator>& OutIndirections, FString* OutError, bool bHandleRedirects) const
+bool FPropertyBindingPath::ResolveIndirectionsWithValue(const FPropertyBindingDataView BaseValueView, TArray<FPropertyBindingPathIndirection, Allocator>& OutIndirections, FString* OutError, bool bHandleRedirects, bool bResetIndirectionsIfFailed) const
 {
 	OutIndirections.Reset();
 	if (OutError)
@@ -265,7 +265,12 @@ bool FPropertyBindingPath::ResolveIndirectionsWithValue(const FPropertyBindingDa
 				*OutError = FString::Printf(TEXT("Malformed path '%s'."),
 					*ToString(Segment.GetIndex(), TEXT("<"), TEXT(">")));
 			}
-			OutIndirections.Reset();
+
+			if (bResetIndirectionsIfFailed)
+			{
+				OutIndirections.Reset();
+			}
+
 			return false;
 		}
 
@@ -364,7 +369,12 @@ bool FPropertyBindingPath::ResolveIndirectionsWithValue(const FPropertyBindingDa
 					*ToString(Segment.GetIndex(), TEXT("<"), TEXT(">")),
 					CurrentStruct->GetPrefixCPP(), *CurrentStruct->GetName(), *Segment->GetName().ToString());
 			}
-			OutIndirections.Reset();
+
+			if (bResetIndirectionsIfFailed)
+			{
+				OutIndirections.Reset();
+			}
+
 			return false;
 		}
 
@@ -404,7 +414,12 @@ bool FPropertyBindingPath::ResolveIndirectionsWithValue(const FPropertyBindingDa
 						*OutError = FString::Printf(TEXT("Index %d out of range (num elements %d) trying to access dynamic array '%s'."),
 							Segment->GetArrayIndex(), Helper.Num(), *ToString(Segment.GetIndex(), TEXT("<"), TEXT(">")));
 					}
-					OutIndirections.Reset();
+
+					if (bResetIndirectionsIfFailed)
+					{
+						OutIndirections.Reset();
+					}
+
 					return false;
 				}
 				CurrentAddress = Helper.GetRawPtr(Segment->GetArrayIndex());
@@ -426,7 +441,12 @@ bool FPropertyBindingPath::ResolveIndirectionsWithValue(const FPropertyBindingDa
 						*OutError = FString::Printf(TEXT("Index %d out of range %d trying to access static array '%s'."),
 							Segment->GetArrayIndex(), Property->ArrayDim, *ToString(Segment.GetIndex(), TEXT("<"), TEXT(">")));
 					}
-					OutIndirections.Reset();
+
+					if (bResetIndirectionsIfFailed)
+					{
+						OutIndirections.Reset();
+					}
+
 					return false;
 				}
 				Offset = Property->GetOffset_ForInternal() + Property->GetElementSize() * FMath::Max(0, Segment->GetArrayIndex());
@@ -459,7 +479,7 @@ bool FPropertyBindingPath::ResolveIndirectionsWithValue(const FPropertyBindingDa
 					// The property path is pointing into the instanced struct, it must be present.
 					// @TODO:	We could potentially check the BaseStruct metadata in editor (for similar behavior as objects)
 					//			Omitting for now to have matching functionality in editor and runtime.
-					auto SetStructInstanceIndirection = [this, &CurrentAddress, &CurrentStruct, &Indirection, &OutError, &OutIndirections]
+					auto SetStructInstanceIndirection = [this, bResetIndirectionsIfFailed, &CurrentAddress, &CurrentStruct, &Indirection, &OutError, &OutIndirections]
 						(const UScriptStruct* ScriptStruct, const uint8* StructMemory, const int32 SegmentIndex, const UStruct* SegmentScriptStruct)
 						{
 							if (ScriptStruct == nullptr || StructMemory == nullptr)
@@ -469,7 +489,12 @@ bool FPropertyBindingPath::ResolveIndirectionsWithValue(const FPropertyBindingDa
 									*OutError = FString::Printf(TEXT("Expecting valid instanced struct value at path '%s'."),
 										*ToString(SegmentIndex, TEXT("<"), TEXT(">")));
 								}
-								OutIndirections.Reset();
+
+								if (bResetIndirectionsIfFailed)
+								{
+									OutIndirections.Reset();
+								}
+
 								return false;
 							}
 
@@ -620,7 +645,12 @@ bool FPropertyBindingPath::ResolveIndirectionsWithValue(const FPropertyBindingDa
 					*OutError = FString::Printf(TEXT("Unsupported property indirection type %s in path '%s'."),
 						*Property->GetCPPType(), *ToString(Segment.GetIndex(), TEXT("<"), TEXT(">")));
 				}
-				OutIndirections.Reset();
+
+				if (bResetIndirectionsIfFailed)
+				{
+					OutIndirections.Reset();
+				}
+
 				return false;
 			}
 		}
@@ -629,14 +659,14 @@ bool FPropertyBindingPath::ResolveIndirectionsWithValue(const FPropertyBindingDa
 	return true;
 }
 
-bool FPropertyBindingPath::ResolveIndirectionsWithValue(const FPropertyBindingDataView BaseValueView, TArray<FPropertyBindingPathIndirection>& OutIndirections, FString* OutError, bool bHandleRedirects) const
+bool FPropertyBindingPath::ResolveIndirectionsWithValue(const FPropertyBindingDataView BaseValueView, TArray<FPropertyBindingPathIndirection>& OutIndirections, FString* OutError, bool bHandleRedirects, bool bResetIndirectionsIfFailed) const
 {
-	return ResolveIndirectionsWithValue<>(BaseValueView, OutIndirections, OutError, bHandleRedirects);
+	return ResolveIndirectionsWithValue<>(BaseValueView, OutIndirections, OutError, bHandleRedirects, bResetIndirectionsIfFailed);
 }
 
-bool FPropertyBindingPath::ResolveIndirectionsWithValue(const FPropertyBindingDataView BaseValueView, TArray<FPropertyBindingPathIndirection, TInlineAllocator<16>>& OutIndirections, FString* OutError, bool bHandleRedirects) const
+bool FPropertyBindingPath::ResolveIndirectionsWithValue(const FPropertyBindingDataView BaseValueView, TArray<FPropertyBindingPathIndirection, TInlineAllocator<16>>& OutIndirections, FString* OutError, bool bHandleRedirects, bool bResetIndirectionsIfFailed) const
 {
-	return ResolveIndirectionsWithValue<>(BaseValueView, OutIndirections, OutError, bHandleRedirects);
+	return ResolveIndirectionsWithValue<>(BaseValueView, OutIndirections, OutError, bHandleRedirects, bResetIndirectionsIfFailed);
 }
 
 bool FPropertyBindingPath::operator==(const FPropertyBindingPath& RHS) const
diff --git a/Engine/Plugins/Runtime/PropertyBindingUtils/Source/PropertyBindingUtils/Public/PropertyBindingPath.h b/Engine/Plugins/Runtime/PropertyBindingUtils/Source/PropertyBindingUtils/Public/PropertyBindingPath.h
--- a/Engine/Plugins/Runtime/PropertyBindingUtils/Source/PropertyBindingUtils/Public/PropertyBindingPath.h
+++ b/Engine/Plugins/Runtime/PropertyBindingUtils/Source/PropertyBindingUtils/Public/PropertyBindingPath.h
@@ -289,10 +289,11 @@ struct FPropertyBindingPath
 	 * @param OutIndirections Indirections describing how the properties were accessed. 
 	 * @param OutError Optional, pointer to string where error will be logged if update fails.
 	 * @param bHandleRedirects If true, the method will try to resolve missing properties using core redirects, and properties on Blueprint and User Defined Structs by ID. Available only in editor builds!
+	 * @param bResetIndirectionsIfFailed If true, indirections resolved so far will be cleared if the path failed to resolve in the end. 
 	 * @return true of path could be resolved against the base value, and instance types were updated.
 	 */
-	PROPERTYBINDINGUTILS_API bool ResolveIndirections(const UStruct* BaseStruct, TArray<FPropertyBindingPathIndirection>& OutIndirections, FString* OutError = nullptr, bool bHandleRedirects = false) const;
-	PROPERTYBINDINGUTILS_API bool ResolveIndirections(const UStruct* BaseStruct, TArray<FPropertyBindingPathIndirection, TInlineAllocator<16>>& OutIndirections, FString* OutError = nullptr, bool bHandleRedirects = false) const;
+	PROPERTYBINDINGUTILS_API bool ResolveIndirections(const UStruct* BaseStruct, TArray<FPropertyBindingPathIndirection>& OutIndirections, FString* OutError = nullptr, bool bHandleRedirects = false, bool bResetIndirectionsIfFailed = true) const;
+	PROPERTYBINDINGUTILS_API bool ResolveIndirections(const UStruct* BaseStruct, TArray<FPropertyBindingPathIndirection, TInlineAllocator<16>>& OutIndirections, FString* OutError = nullptr, bool bHandleRedirects = false, bool bResetIndirectionsIfFailed = true) const;
 
 	/**
 	 * Resolves the property path against base value. The path is assumed to be relative to the BaseValueView.
@@ -300,10 +301,11 @@ struct FPropertyBindingPath
 	 * @param OutIndirections Indirections describing how the properties were accessed. 
 	 * @param OutError Optional, pointer to string where error will be logged if update fails.
 	 * @param bHandleRedirects If true, the method will try to resolve missing properties using core redirects, and properties on Blueprint and User Defined Structs by ID. Available only in editor builds!
+	 * @param bResetIndirectionsIfFailed If true, indirections resolved so far will be cleared if the path failed to resolve in the end.
 	 * @return true of path could be resolved against the base value, and instance types were updated.
 	 */
-	PROPERTYBINDINGUTILS_API bool ResolveIndirectionsWithValue(const FPropertyBindingDataView BaseValueView, TArray<FPropertyBindingPathIndirection>& OutIndirections, FString* OutError = nullptr, bool bHandleRedirects = false) const;
-	PROPERTYBINDINGUTILS_API bool ResolveIndirectionsWithValue(const FPropertyBindingDataView BaseValueView, TArray<FPropertyBindingPathIndirection, TInlineAllocator<16>>& OutIndirections, FString* OutError = nullptr, bool bHandleRedirects = false) const;
+	PROPERTYBINDINGUTILS_API bool ResolveIndirectionsWithValue(const FPropertyBindingDataView BaseValueView, TArray<FPropertyBindingPathIndirection>& OutIndirections, FString* OutError = nullptr, bool bHandleRedirects = false, bool bResetIndirectionsIfFailed = true) const;
+	PROPERTYBINDINGUTILS_API bool ResolveIndirectionsWithValue(const FPropertyBindingDataView BaseValueView, TArray<FPropertyBindingPathIndirection, TInlineAllocator<16>>& OutIndirections, FString* OutError = nullptr, bool bHandleRedirects = false, bool bResetIndirectionsIfFailed = true) const;
 
 	/**
 	 * Updates property segments from base struct type. The path is expected to be relative to the BaseStruct.
@@ -416,7 +418,7 @@ private:
 	TArray<FPropertyBindingPathSegment> Segments;
 
 	template<typename Allocator>
-	bool ResolveIndirectionsWithValue(const FPropertyBindingDataView BaseValueView, TArray<FPropertyBindingPathIndirection, Allocator>& OutIndirections, FString* OutError = nullptr, bool bHandleRedirects = false) const;
+	bool ResolveIndirectionsWithValue(const FPropertyBindingDataView BaseValueView, TArray<FPropertyBindingPathIndirection, Allocator>& OutIndirections, FString* OutError = nullptr, bool bHandleRedirects = false, bool bResetIndirectionsIfFailed = true) const;
 };
 
 
diff --git a/Engine/Plugins/Runtime/StateTree/Source/StateTreeEditorModule/Private/StateTreeEditorMode.cpp b/Engine/Plugins/Runtime/StateTree/Source/StateTreeEditorModule/Private/StateTreeEditorMode.cpp
--- a/Engine/Plugins/Runtime/StateTree/Source/StateTreeEditorModule/Private/StateTreeEditorMode.cpp
+++ b/Engine/Plugins/Runtime/StateTree/Source/StateTreeEditorModule/Private/StateTreeEditorMode.cpp
@@ -571,17 +571,20 @@ void UStateTreeEditorMode::HandleModelBringBindingPathToFocus(const UStateTreeSt
 
 	auto MakeHighlightPathFromIndirections = [](FStateTreeDataView InValue, const FPropertyBindingPath& InPropertyBindingPath, FPropertyPath& OutHighlightPath, int32 InStartIndex = 0)
 		{
+			check(InStartIndex >= 0);
+
 			// We want it to be in sync with what we show
 			constexpr bool bHandleRedirects = false;
+			// We want to have partial path to find the most match one
+			constexpr bool bResetIndirectionsIfFailed = false;
 			TArray<FPropertyBindingPathIndirection> Indirections;
-			if (InPropertyBindingPath.ResolveIndirectionsWithValue(InValue, Indirections, nullptr, bHandleRedirects))
+
+			InPropertyBindingPath.ResolveIndirectionsWithValue(InValue, Indirections, nullptr, bHandleRedirects, bResetIndirectionsIfFailed);
+			for (int32 Idx = InStartIndex; Idx < Indirections.Num(); ++Idx)
 			{
-				for (int32 Idx = InStartIndex; Idx < Indirections.Num(); ++Idx)
-				{
-					const FPropertyBindingPathIndirection& Indirection = Indirections[Idx];
-					OutHighlightPath.AddProperty(FPropertyInfo(const_cast<FProperty*>(Indirection.GetProperty()), Indirection.GetArrayIndex()));
-				}
-			};
+				const FPropertyBindingPathIndirection& Indirection = Indirections[Idx];
+				OutHighlightPath.AddProperty(FPropertyInfo(const_cast<FProperty*>(Indirection.GetProperty()), Indirection.GetArrayIndex()));
+			}
 		};
 
 	auto MakeHighlightPathFromEvent = [&MakeHighlightPathFromIndirections](FProperty* InEventProperty, const FStateTreeBindableStructDesc& InDesc, FStateTreeDataView InValue, const FPropertyBindingPath& InPropertyBindingPath, FPropertyPath& OutHighlightPath)
