[PropertyBindingUtils] Use ArePropertyGuidsAvailable to find the property when using bHandleredirects. Similar to how SerializeVersionedTaggedProperties finds the property.

#Synced-CL 50535091

diff --git a/Engine/Plugins/Runtime/PropertyBindingUtils/Source/PropertyBindingUtils/Private/PropertyBindingPath.cpp b/Engine/Plugins/Runtime/PropertyBindingUtils/Source/PropertyBindingUtils/Private/PropertyBindingPath.cpp
--- a/Engine/Plugins/Runtime/PropertyBindingUtils/Source/PropertyBindingUtils/Private/PropertyBindingPath.cpp
+++ b/Engine/Plugins/Runtime/PropertyBindingUtils/Source/PropertyBindingUtils/Private/PropertyBindingPath.cpp
@@ -274,7 +274,8 @@ bool FPropertyBindingPath::ResolveIndirectionsWithValue(const FPropertyBindingDa
 			return false;
 		}
 
-		const FProperty* Property = CurrentStruct->FindPropertyByName(Segment->GetName());
+		FName PropertyName = Segment->GetName();
+		const FProperty* Property = CurrentStruct->FindPropertyByName(PropertyName);
 		const bool bWithValue = CurrentAddress != nullptr;
 
 #if WITH_EDITORONLY_DATA
@@ -284,83 +285,47 @@ bool FPropertyBindingPath::ResolveIndirectionsWithValue(const FPropertyBindingDa
 		// Try to fix the path in editor.
 		if (bHandleRedirects)
 		{
-			
-			// Check if there's a core redirect for it.
-			if (!Property)
+			if (!Property && PropertyGuid.IsValid())
 			{
-				// Try to match by property ID (Blueprint or User Defined Struct).
-				if (Segment->GetPropertyGuid().IsValid())
+				// Try to match by property ID (Blueprint, User Defined Struct, PropertyBag).
+				if (CurrentStruct->ArePropertyGuidsAvailable())
 				{
-					if (const UBlueprintGeneratedClass* BlueprintClass = Cast<UBlueprintGeneratedClass>(CurrentStruct))
-					{
-						if (const FName* Name = BlueprintClass->PropertyGuids.FindKey(Segment->GetPropertyGuid()))
-						{
-							RedirectedName = *Name;
-							Property = CurrentStruct->FindPropertyByName(RedirectedName);
-						}
-					}
-					else if (const UUserDefinedStruct* UserDefinedStruct = Cast<UUserDefinedStruct>(CurrentStruct))
-					{
-						if (FProperty* FoundProperty = FStructureEditorUtils::GetPropertyByGuid(UserDefinedStruct, Segment->GetPropertyGuid()))
-						{
-							RedirectedName = FoundProperty->GetFName();
-							Property = FoundProperty;
-						}
-					}
-					else if (const UPropertyBag* PropertyBag = Cast<UPropertyBag>(CurrentStruct))
+					const FName NewName = CurrentStruct->FindPropertyNameFromGuid(PropertyGuid);
+					if (!NewName.IsNone() && NewName != PropertyName)
 					{
-						if (const FPropertyBagPropertyDesc* Desc = PropertyBag->FindPropertyDescByID(Segment->GetPropertyGuid()))
-						{
-							if (Desc->CachedProperty)
-							{
-								RedirectedName = Desc->CachedProperty->GetFName();
-								Property = Desc->CachedProperty;
-							}
-						}
-					}
-				}
-				else
-				{
-					// Try core redirect
-					const FCoreRedirectObjectName OldPropertyName(Segment->GetName(), CurrentStruct->GetFName(), *CurrentStruct->GetOutermost()->GetPathName());
-					const FCoreRedirectObjectName NewPropertyName = FCoreRedirects::GetRedirectedName(ECoreRedirectFlags::Type_Property, OldPropertyName);
-					if (OldPropertyName != NewPropertyName)
-					{
-						// Cached the result for later use.
-						RedirectedName = NewPropertyName.ObjectName;
-
-						Property = CurrentStruct->FindPropertyByName(RedirectedName);
+						RedirectedName = NewName;
+						PropertyName = NewName;
+						Property = CurrentStruct->FindPropertyByName(PropertyName);
 					}
 				}
 			}
 
-			// Update PropertyGuid 
-			if (Property)
+			// Check if there's a core redirect for it.
+			if (!Property)
 			{
-				const FName PropertyName = !RedirectedName.IsNone() ? RedirectedName : Segment->GetName();
-				if (const UBlueprintGeneratedClass* BlueprintClass = Cast<UBlueprintGeneratedClass>(CurrentStruct))
+				// Try core redirect
+				for (const UStruct* CheckStruct : CurrentStruct->GetSuperStructIterator())
 				{
-					if (const FGuid* VarGuid = BlueprintClass->PropertyGuids.Find(PropertyName))
+					const FName NewPropertyName = FProperty::FindRedirectedPropertyName(CheckStruct, PropertyName);
+					if (!NewPropertyName.IsNone())
 					{
-						PropertyGuid = *VarGuid;
-					}
-				}
-				else if (const UUserDefinedStruct* UserDefinedStruct = Cast<UUserDefinedStruct>(CurrentStruct))
-				{
-					// Parse Guid from UDS property name.
-					PropertyGuid = FStructureEditorUtils::GetGuidFromPropertyName(PropertyName);
-				}
-				else if (const UPropertyBag* PropertyBag = Cast<UPropertyBag>(CurrentStruct))
-				{
-					if (const FPropertyBagPropertyDesc* Desc = PropertyBag->FindPropertyDescByPropertyName(PropertyName))
-					{
-						PropertyGuid = Desc->ID;
+						// Cached the result for later use.
+						RedirectedName = NewPropertyName;
+						PropertyName = RedirectedName;
+						Property = CurrentStruct->FindPropertyByName(PropertyName);
+						break;
 					}
 				}
 			}
 		}
 #endif // WITH_EDITORONLY_DATA
 
+		// Check if it has a custom name
+		if (!Property)
+		{
+			Property = CurrentStruct->CustomFindProperty(PropertyName);
+		}
+
 		if (!Property)
 		{
 			if (OutError)
@@ -378,6 +343,14 @@ bool FPropertyBindingPath::ResolveIndirectionsWithValue(const FPropertyBindingDa
 			return false;
 		}
 
+#if WITH_EDITORONLY_DATA
+		// Update PropertyGuid
+		if (bHandleRedirects && CurrentStruct->ArePropertyGuidsAvailable())
+		{
+			PropertyGuid = CurrentStruct->FindPropertyGuidFromName(PropertyName);
+		}
+#endif // WITH_EDITORONLY_DATA
+
 		const FArrayProperty* ArrayProperty = CastField<FArrayProperty>(Property);
 		int ArrayIndex = INDEX_NONE;
 		int32 Offset = INDEX_NONE;
@@ -399,7 +372,7 @@ bool FPropertyBindingPath::ResolveIndirectionsWithValue(const FPropertyBindingDa
 			Indirection.RedirectedName = RedirectedName;
 			Indirection.PropertyGuid = PropertyGuid;
 #endif
-			
+
 			ArrayIndex = INDEX_NONE;
 			Offset = 0;
 			Property = ArrayProperty->Inner;
