Property Binding Extension: Make detail widget wrapping logic overridable.

#rnx

#Synced-CL 50938845

diff --git a/Engine/Plugins/Runtime/PropertyBindingUtils/Source/PropertyBindingUtilsEditor/Private/PropertyBindingExtension.cpp b/Engine/Plugins/Runtime/PropertyBindingUtils/Source/PropertyBindingUtilsEditor/Private/PropertyBindingExtension.cpp
--- a/Engine/Plugins/Runtime/PropertyBindingUtils/Source/PropertyBindingUtilsEditor/Private/PropertyBindingExtension.cpp
+++ b/Engine/Plugins/Runtime/PropertyBindingUtils/Source/PropertyBindingUtilsEditor/Private/PropertyBindingExtension.cpp
@@ -1298,22 +1298,6 @@ void FPropertyBindingExtension::ExtendWidgetRow(FDetailWidgetRow& InWidgetRow, c
 
 	TSharedPtr<FCachedBindingData> CachedBindingData = CreateCachedBindingData(BindingsOwner, TargetPath, InPropertyHandle, AccessibleStructs);
 
-	// Wrap value widget
-	auto IsValueVisible = TAttribute<EVisibility>::Create([CachedBindingData]() -> EVisibility
-	{
-		return CachedBindingData->HasBinding(FPropertyBindingBindingCollection::ESearchMode::Exact) ? EVisibility::Collapsed : EVisibility::Visible;
-	});
-
-	TSharedPtr<SWidget> ValueWidget = InWidgetRow.ValueContent().Widget;
-	InWidgetRow.ValueContent()
-	[
-		SNew(SBox)
-		.Visibility(IsValueVisible)
-		[
-			ValueWidget.ToSharedRef()
-		]
-	];
-
 	FPropertyBindingWidgetArgs Args;
 	Args.Property = InPropertyHandle->GetProperty();
 
@@ -1512,27 +1496,51 @@ void FPropertyBindingExtension::UpdateContextStruct(TConstStructView<FPropertyBi
 
 void FPropertyBindingExtension::CustomizeDetailWidgetRow(FDetailWidgetRow& InWidgetRow, const IDetailLayoutBuilder& InDetailBuilder, IPropertyBindingBindingCollectionOwner* InBindingsOwner, TSharedPtr<IPropertyHandle> InPropertyHandle, const FPropertyBindingPath& InTargetPath, TSharedPtr<UE::PropertyBinding::FCachedBindingData> InCachedBindingData) const
 {
-	// ResetToDefault
+	WrapValueWidget(InWidgetRow, InCachedBindingData);
+	CustomizeResetToDefault(InWidgetRow, InDetailBuilder, InPropertyHandle, InCachedBindingData);
+}
+
+void FPropertyBindingExtension::WrapValueWidget(FDetailWidgetRow& InWidgetRow,
+	const TSharedPtr<UE::PropertyBinding::FCachedBindingData>& InCachedBindingData) const
+{
+	auto IsValueVisible = TAttribute<EVisibility>::Create([InCachedBindingData]() -> EVisibility
 	{
-		InWidgetRow.CustomResetToDefault = FResetToDefaultOverride::Create(
-			MakeAttributeLambda([InCachedBindingData, InPropertyHandle]()
+		return InCachedBindingData->HasBinding(FPropertyBindingBindingCollection::ESearchMode::Exact) ? EVisibility::Collapsed : EVisibility::Visible;
+	});
+
+	TSharedPtr<SWidget> ValueWidget = InWidgetRow.ValueContent().Widget;
+	InWidgetRow.ValueContent()
+	[
+		SNew(SBox)
+		.Visibility(IsValueVisible)
+		[
+			ValueWidget.ToSharedRef()
+		]
+	];
+}
+
+void FPropertyBindingExtension::CustomizeResetToDefault(FDetailWidgetRow& InWidgetRow,
+	const IDetailLayoutBuilder& InDetailBuilder, const TSharedPtr<IPropertyHandle>& InPropertyHandle,
+	const TSharedPtr<UE::PropertyBinding::FCachedBindingData>& InCachedBindingData) const
+{
+	InWidgetRow.CustomResetToDefault = FResetToDefaultOverride::Create(
+		MakeAttributeLambda([InCachedBindingData, InPropertyHandle]()
+			{
+				return InPropertyHandle->CanResetToDefault() || InCachedBindingData->HasBinding(FPropertyBindingBindingCollection::ESearchMode::Includes);
+			}),
+		FSimpleDelegate::CreateLambda([InCachedBindingData, &InDetailBuilder, InPropertyHandle]()
+			{
+				if (InCachedBindingData->HasBinding(FPropertyBindingBindingCollection::ESearchMode::Includes))
 				{
-					return InPropertyHandle->CanResetToDefault() || InCachedBindingData->HasBinding(FPropertyBindingBindingCollection::ESearchMode::Includes);
-				}),
-			FSimpleDelegate::CreateLambda([InCachedBindingData, &InDetailBuilder, InPropertyHandle]()
+					InCachedBindingData->RemoveBinding(FPropertyBindingBindingCollection::ESearchMode::Includes);
+					InDetailBuilder.GetPropertyUtilities()->RequestForceRefresh();
+				}
+				if (InPropertyHandle->CanResetToDefault())
 				{
-					if (InCachedBindingData->HasBinding(FPropertyBindingBindingCollection::ESearchMode::Includes))
-					{
-						InCachedBindingData->RemoveBinding(FPropertyBindingBindingCollection::ESearchMode::Includes);
-						InDetailBuilder.GetPropertyUtilities()->RequestForceRefresh();
-					}
-					if (InPropertyHandle->CanResetToDefault())
-					{
-						InPropertyHandle->ResetToDefault();
-					}
-				}),
-			false);
-	}
+					InPropertyHandle->ResetToDefault();
+				}
+			}),
+		false);
 }
 
 void FPropertyBindingExtension::CustomizePropertyBindingWidgetArgs(FPropertyBindingWidgetArgs& Args,
diff --git a/Engine/Plugins/Runtime/PropertyBindingUtils/Source/PropertyBindingUtilsEditor/Public/PropertyBindingExtension.h b/Engine/Plugins/Runtime/PropertyBindingUtils/Source/PropertyBindingUtilsEditor/Public/PropertyBindingExtension.h
--- a/Engine/Plugins/Runtime/PropertyBindingUtils/Source/PropertyBindingUtilsEditor/Public/PropertyBindingExtension.h
+++ b/Engine/Plugins/Runtime/PropertyBindingUtils/Source/PropertyBindingUtilsEditor/Public/PropertyBindingExtension.h
@@ -280,6 +280,9 @@ protected:
 	PROPERTYBINDINGUTILSEDITOR_API virtual bool CanBindToArrayElements() const;
 	[[nodiscard]] PROPERTYBINDINGUTILSEDITOR_API virtual TNotNull<const UEdGraphSchema_K2*> GetEdGraphSchema() const;
 
+	PROPERTYBINDINGUTILSEDITOR_API void WrapValueWidget(FDetailWidgetRow& InWidgetRow, const TSharedPtr<UE::PropertyBinding::FCachedBindingData>& InCachedBindingData) const;
+	PROPERTYBINDINGUTILSEDITOR_API void CustomizeResetToDefault(FDetailWidgetRow& InWidgetRow, const IDetailLayoutBuilder& InDetailBuilder, const TSharedPtr<IPropertyHandle>& InPropertyHandle, const TSharedPtr<UE::PropertyBinding::FCachedBindingData>& InCachedBindingData) const;
+
 	UE_DEPRECATED(5.8, "Use the overload with Schema instead.")
 	PROPERTYBINDINGUTILSEDITOR_API static TSharedRef<SWidget> MakeContextStructWidget(const FPropertyBindingBindableStructDescriptor& InContextStruct);
 	PROPERTYBINDINGUTILSEDITOR_API static TSharedRef<SWidget> MakeContextStructWidget(TNotNull<const UEdGraphSchema_K2*> Schema, const FPropertyBindingBindableStructDescriptor& InContextStruct);
