[State Tree] Disable promoting to state parameters if the state is a linked state

#Synced-CL 51004728

diff --git a/Engine/Plugins/Runtime/PropertyBindingUtils/Source/PropertyBindingUtilsEditor/Private/PropertyBindingExtension.cpp b/Engine/Plugins/Runtime/PropertyBindingUtils/Source/PropertyBindingUtilsEditor/Private/PropertyBindingExtension.cpp
--- a/Engine/Plugins/Runtime/PropertyBindingUtils/Source/PropertyBindingUtilsEditor/Private/PropertyBindingExtension.cpp
+++ b/Engine/Plugins/Runtime/PropertyBindingUtils/Source/PropertyBindingUtilsEditor/Private/PropertyBindingExtension.cpp
@@ -1483,13 +1483,14 @@ TSharedPtr<UE::PropertyBinding::FCachedBindingData> FPropertyBindingExtension::C
 	return MakeShared<UE::PropertyBinding::FCachedBindingData>(InBindingsOwner, InTargetPath, InPropertyHandle, InAccessibleStructs);
 }
 
-bool FPropertyBindingExtension::GetPromotionToParameterOverrideInternal(const FProperty& InProperty, bool& bOutOverride) const
+bool FPropertyBindingExtension::GetPromotionToParameterOverrideInternal(const TSharedRef<const IPropertyHandle> InPropertyHandle,
+	bool& bOutOverride) const
 {
 	return false;
 }
 
 void FPropertyBindingExtension::UpdateContextStruct(TConstStructView<FPropertyBindingBindableStructDescriptor> InStructDesc, FBindingContextStruct& InOutContextStruct
-	, TMap<FString, FText>& InOutSectionNames) const
+                                                    , TMap<FString, FText>& InOutSectionNames) const
 {
 	// nothing to do
 }
@@ -1568,7 +1569,7 @@ bool FPropertyBindingExtension::CanPromoteToParameter(const TSharedPtr<IProperty
 		return false;
 	}
 
-	if (bool bCanPromoteOverride = false; GetPromotionToParameterOverrideInternal(*Property, bCanPromoteOverride))
+	if (bool bCanPromoteOverride = false; GetPromotionToParameterOverrideInternal(InPropertyHandle.ToSharedRef(), bCanPromoteOverride))
 	{
 		return bCanPromoteOverride;
 	}
diff --git a/Engine/Plugins/Runtime/PropertyBindingUtils/Source/PropertyBindingUtilsEditor/Public/PropertyBindingExtension.h b/Engine/Plugins/Runtime/PropertyBindingUtils/Source/PropertyBindingUtilsEditor/Public/PropertyBindingExtension.h
--- a/Engine/Plugins/Runtime/PropertyBindingUtils/Source/PropertyBindingUtilsEditor/Public/PropertyBindingExtension.h
+++ b/Engine/Plugins/Runtime/PropertyBindingUtils/Source/PropertyBindingUtilsEditor/Public/PropertyBindingExtension.h
@@ -259,13 +259,19 @@ protected:
 	PROPERTYBINDINGUTILSEDITOR_API virtual bool CanBindToProperty(const FPropertyBindingPath& InTargetPath, const IPropertyHandle& InPropertyHandle) const;
 	PROPERTYBINDINGUTILSEDITOR_API virtual uint64 GetDisallowedPropertyFlags() const;
 
+	UE_DEPRECATED(5.8, "Use the version with PropertyHandle instead.")
+	virtual bool GetPromotionToParameterOverrideInternal(const FProperty& InProperty, bool& bOutOverride) const final
+	{
+		return false;
+	}
+
 	/**
-	 * The optional override criteria of whether the property can be promoted to parameter 
-	 * @param InProperty the property to be tested upon
+	 * The optional override criteria of whether the property can be promoted to parameter
+	 * @param InPropertyHandle the property handle to be tested upon
 	 * @param bOutOverride Whether the property should be promoted to parameter, if the return value is true
 	 * @return Whether the promotion criteria is overridden for this property
 	 */
-	PROPERTYBINDINGUTILSEDITOR_API virtual bool GetPromotionToParameterOverrideInternal(const FProperty& InProperty, bool& bOutOverride) const;
+	PROPERTYBINDINGUTILSEDITOR_API virtual bool GetPromotionToParameterOverrideInternal(const TSharedRef<const IPropertyHandle> InPropertyHandle, bool& bOutOverride) const;
 	PROPERTYBINDINGUTILSEDITOR_API virtual void UpdateContextStruct(TConstStructView<FPropertyBindingBindableStructDescriptor> InStructDesc, FBindingContextStruct& InOutContextStruct, TMap<FString, FText>& InOutSectionNames) const;
 	PROPERTYBINDINGUTILSEDITOR_API virtual void CustomizeDetailWidgetRow(FDetailWidgetRow& InWidgetRow, const IDetailLayoutBuilder& InDetailBuilder, IPropertyBindingBindingCollectionOwner* InBindingsOwner, TSharedPtr<IPropertyHandle> InPropertyHandle, const FPropertyBindingPath& InTargetPath, TSharedPtr<UE::PropertyBinding::FCachedBindingData> InCachedBindingData) const;
 
diff --git a/Engine/Plugins/Runtime/StateTree/Source/StateTreeEditorModule/Private/Customizations/StateTreeBindingExtension.cpp b/Engine/Plugins/Runtime/StateTree/Source/StateTreeEditorModule/Private/Customizations/StateTreeBindingExtension.cpp
--- a/Engine/Plugins/Runtime/StateTree/Source/StateTreeEditorModule/Private/Customizations/StateTreeBindingExtension.cpp
+++ b/Engine/Plugins/Runtime/StateTree/Source/StateTreeEditorModule/Private/Customizations/StateTreeBindingExtension.cpp
@@ -825,9 +825,32 @@ void FStateTreeBindingExtension::UpdateContextStruct(TConstStructView<FPropertyB
 	}
 }
 
-bool FStateTreeBindingExtension::GetPromotionToParameterOverrideInternal(const FProperty& InProperty, bool& bOutOverride) const
+bool FStateTreeBindingExtension::GetPromotionToParameterOverrideInternal(const TSharedRef<const IPropertyHandle> InPropertyHandle, bool& bOutOverride) const
 {
-	if (const FStructProperty* StructProperty = CastField<FStructProperty>(&InProperty))
+	TSharedPtr<const IPropertyHandle> CurrentHandle = InPropertyHandle.ToSharedPtr();
+	while (const TSharedPtr<const IPropertyHandle> ParentHandle = CurrentHandle->GetParentHandle())
+	{
+		CurrentHandle = ParentHandle;
+	}
+
+	if (CurrentHandle.IsValid())
+	{
+		TArray<UObject*> RootObjects;
+		CurrentHandle->GetOuterObjects(RootObjects);
+		if (RootObjects.Num() == 1)
+		{
+			if (const UStateTreeState* State = Cast<UStateTreeState>(RootObjects[0]))
+			{
+				if (State->Type == EStateTreeStateType::Linked || State->Type == EStateTreeStateType::LinkedAsset)
+				{
+					bOutOverride = false;
+					return true;
+				}
+			}
+		}
+	}
+
+	if (const FStructProperty* StructProperty = CastField<FStructProperty>(InPropertyHandle->GetProperty()))
 	{
 		// Support Property Refs as even though these aren't bp types, the actual types that would be added are the ones in the meta-data RefType
 		if (StructProperty->Struct && StructProperty->Struct->IsChildOf(FStateTreePropertyRef::StaticStruct()))
diff --git a/Engine/Plugins/Runtime/StateTree/Source/StateTreeEditorModule/Private/Customizations/StateTreeBindingExtension.h b/Engine/Plugins/Runtime/StateTree/Source/StateTreeEditorModule/Private/Customizations/StateTreeBindingExtension.h
--- a/Engine/Plugins/Runtime/StateTree/Source/StateTreeEditorModule/Private/Customizations/StateTreeBindingExtension.h
+++ b/Engine/Plugins/Runtime/StateTree/Source/StateTreeEditorModule/Private/Customizations/StateTreeBindingExtension.h
@@ -25,7 +25,7 @@ class FStateTreeBindingExtension : public FPropertyBindingExtension
 private:
 	virtual TSharedPtr<UE::PropertyBinding::FCachedBindingData> CreateCachedBindingData(IPropertyBindingBindingCollectionOwner* InBindingsOwner, const FPropertyBindingPath& InTargetPath, const TSharedPtr<IPropertyHandle>& InPropertyHandle, TConstArrayView<TInstancedStruct<FPropertyBindingBindableStructDescriptor>> InAccessibleStructs) const override;
 	virtual bool CanBindToProperty(const FPropertyBindingPath& InTargetPath, const IPropertyHandle& InPropertyHandle) const override;
-	virtual bool GetPromotionToParameterOverrideInternal(const FProperty& InProperty, bool& bOutOverride) const override;
+	virtual bool GetPromotionToParameterOverrideInternal(const TSharedRef<const IPropertyHandle> InPropertyHandle, bool& bOutOverride) const override;
 	virtual void UpdateContextStruct(TConstStructView<FPropertyBindingBindableStructDescriptor> InStructDesc,FBindingContextStruct& InOutContextStruct, TMap<FString, FText>& InOutSectionNames) const override;
 	virtual bool CanBindToArrayElements() const override
 	{
diff --git a/Engine/Plugins/VirtualProduction/SceneState/Source/SceneStateBlueprintEditor/Private/DetailsView/SceneStatePropertyBindingExtension.cpp b/Engine/Plugins/VirtualProduction/SceneState/Source/SceneStateBlueprintEditor/Private/DetailsView/SceneStatePropertyBindingExtension.cpp
--- a/Engine/Plugins/VirtualProduction/SceneState/Source/SceneStateBlueprintEditor/Private/DetailsView/SceneStatePropertyBindingExtension.cpp
+++ b/Engine/Plugins/VirtualProduction/SceneState/Source/SceneStateBlueprintEditor/Private/DetailsView/SceneStatePropertyBindingExtension.cpp
@@ -56,9 +56,9 @@ TSharedPtr<UE::PropertyBinding::FCachedBindingData> FBindingExtension::CreateCac
 	return MakeShared<FSceneStateCachedBindingData>(InBindingsOwner, InTargetPath, InPropertyHandle, InAccessibleStructs);
 }
 
-bool FBindingExtension::GetPromotionToParameterOverrideInternal(const FProperty& InProperty, bool& bOutOverride) const
+bool FBindingExtension::GetPromotionToParameterOverrideInternal(const TSharedRef<const IPropertyHandle> InPropertyHandle, bool& bOutOverride) const
 {
-	if (const FStructProperty* StructProperty = CastField<FStructProperty>(&InProperty))
+	if (const FStructProperty* StructProperty = CastField<FStructProperty>(InPropertyHandle->GetProperty()))
 	{
 		// Support Property Refs as even though these aren't bp types, the actual types that would be added are the ones in the meta-data RefType
 		if (StructProperty->Struct && StructProperty->Struct->IsChildOf<FSceneStatePropertyReference>())
diff --git a/Engine/Plugins/VirtualProduction/SceneState/Source/SceneStateBlueprintEditor/Private/DetailsView/SceneStatePropertyBindingExtension.h b/Engine/Plugins/VirtualProduction/SceneState/Source/SceneStateBlueprintEditor/Private/DetailsView/SceneStatePropertyBindingExtension.h
--- a/Engine/Plugins/VirtualProduction/SceneState/Source/SceneStateBlueprintEditor/Private/DetailsView/SceneStatePropertyBindingExtension.h
+++ b/Engine/Plugins/VirtualProduction/SceneState/Source/SceneStateBlueprintEditor/Private/DetailsView/SceneStatePropertyBindingExtension.h
@@ -17,7 +17,7 @@ public:
 	//~ Begin FPropertyBindingExtension
 	virtual bool IsPropertyExtendable(const UClass* InObjectClass, const IPropertyHandle& InPropertyHandle) const override;
 	virtual TSharedPtr<UE::PropertyBinding::FCachedBindingData> CreateCachedBindingData(IPropertyBindingBindingCollectionOwner* InBindingsOwner, const FPropertyBindingPath& InTargetPath, const TSharedPtr<IPropertyHandle>& InPropertyHandle, TConstArrayView<TInstancedStruct<FPropertyBindingBindableStructDescriptor>> InAccessibleStructs) const override;
-	virtual bool GetPromotionToParameterOverrideInternal(const FProperty& InProperty, bool& bOutOverride) const override;
+	virtual bool GetPromotionToParameterOverrideInternal(const TSharedRef<const IPropertyHandle> InPropertyHandle, bool& bOutOverride) const override;
 	virtual void UpdateContextStruct(TConstStructView<FPropertyBindingBindableStructDescriptor> InStructDesc, FBindingContextStruct& InOutContextStruct, TMap<FString, FText>& InOutSectionNames) const override;
 	virtual bool CanBindToProperty(const FPropertyBindingPath& InTargetPath, const IPropertyHandle& InPropertyHandle) const override;
 	virtual bool CanBindToArrayElements() const override;
