[PropertyBindingUtils] added an dummy indirection with container and segment index info for ResolveIndirections when indirections are not set to reset on fail

A binding path Foo.Bar.X becomes invalid when variable X doesn't exist anymore. But callers may still want to retain the information of valid indirections so far. Add a dummy indirection to indicate the container struct Bar even X is invalid.

#Synced-CL 51266215

diff --git a/Engine/Plugins/Runtime/PropertyBindingUtils/Source/PropertyBindingUtils/Private/PropertyBindingPath.cpp b/Engine/Plugins/Runtime/PropertyBindingUtils/Source/PropertyBindingUtils/Private/PropertyBindingPath.cpp
--- a/Engine/Plugins/Runtime/PropertyBindingUtils/Source/PropertyBindingUtils/Private/PropertyBindingPath.cpp
+++ b/Engine/Plugins/Runtime/PropertyBindingUtils/Source/PropertyBindingUtils/Private/PropertyBindingPath.cpp
@@ -322,6 +322,12 @@ bool FPropertyBindingPath::ResolveIndirectionsWithValue(const FPropertyBindingDa
 			{
 				OutIndirections.Reset();
 			}
+			else
+			{
+				// A binding path Foo.Bar.X becomes invalid when variable X doesn't exist anymore. Add a dummy indirection to indicate the segment index of X.
+				FPropertyBindingPathIndirection& DummyIndirection = OutIndirections.Emplace_GetRef();
+				DummyIndirection.PathSegmentIndex = Segment.GetIndex();
+			}
 
 			return false;
 		}
@@ -404,6 +410,14 @@ bool FPropertyBindingPath::ResolveIndirectionsWithValue(const FPropertyBindingDa
 			{
 				OutIndirections.Reset();
 			}
+			else
+			{
+				// A binding path Foo.Bar.X becomes invalid when variable X doesn't exist anymore. Add a dummy indirection to indicate the container struct Bar even X is invalid.
+				FPropertyBindingPathIndirection& DummyIndirection = OutIndirections.Emplace_GetRef();
+				DummyIndirection.ContainerStruct = CurrentStruct;
+				DummyIndirection.ContainerAddress = CurrentAddress;
+				DummyIndirection.PathSegmentIndex = Segment.GetIndex();
+			}
 
 			return false;
 		}
diff --git a/Engine/Plugins/Runtime/StateTree/Source/StateTreeEditorModule/Private/StateTreeEditorMode.cpp b/Engine/Plugins/Runtime/StateTree/Source/StateTreeEditorModule/Private/StateTreeEditorMode.cpp
--- a/Engine/Plugins/Runtime/StateTree/Source/StateTreeEditorModule/Private/StateTreeEditorMode.cpp
+++ b/Engine/Plugins/Runtime/StateTree/Source/StateTreeEditorModule/Private/StateTreeEditorMode.cpp
@@ -583,7 +583,10 @@ void UStateTreeEditorMode::HandleModelBringBindingPathToFocus(const UStateTreeSt
 			for (int32 Idx = InStartIndex; Idx < Indirections.Num(); ++Idx)
 			{
 				const FPropertyBindingPathIndirection& Indirection = Indirections[Idx];
-				OutHighlightPath.AddProperty(FPropertyInfo(const_cast<FProperty*>(Indirection.GetProperty()), Indirection.GetArrayIndex()));
+				if (FProperty* Property = const_cast<FProperty*>(Indirection.GetProperty()))
+				{
+					OutHighlightPath.AddProperty(FPropertyInfo(Property, Indirection.GetArrayIndex()));
+				}
 			}
 		};
 
