PropertyBindings: Provide a detailed error message when promoting a parameter ensures due to failed indirection resolve.

#rnx
#Synced-CL 52298949

diff --git a/Engine/Plugins/Runtime/PropertyBindingUtils/Source/PropertyBindingUtilsEditor/Private/PropertyBindingExtension.cpp b/Engine/Plugins/Runtime/PropertyBindingUtils/Source/PropertyBindingUtilsEditor/Private/PropertyBindingExtension.cpp
--- a/Engine/Plugins/Runtime/PropertyBindingUtils/Source/PropertyBindingUtilsEditor/Private/PropertyBindingExtension.cpp
+++ b/Engine/Plugins/Runtime/PropertyBindingUtils/Source/PropertyBindingUtilsEditor/Private/PropertyBindingExtension.cpp
@@ -861,7 +861,10 @@ void FCachedBindingData::PromoteToParameter(const FName InPropertyName, TConstSt
 	if (BindingsOwner->GetBindingDataViewByID(TargetPath.GetStructID(), TargetDataView) && TargetDataView.IsValid())
 	{
 		TArray<FPropertyBindingPathIndirection> TargetIndirections;
-		if (ensure(TargetPath.ResolveIndirectionsWithValue(TargetDataView, TargetIndirections)))
+		FString Error;
+		const bool bResolved = TargetPath.ResolveIndirectionsWithValue(TargetDataView, TargetIndirections, &Error);
+
+		if (ensureMsgf(bResolved, TEXT("Failed to resolve indirections: %s"), *Error))
 		{
 			const FPropertyBindingPathIndirection& LastIndirection = TargetIndirections.Last();
 			TargetProperty = LastIndirection.GetProperty();
